fixed missing module dependency in UnrealLightmass on Mac

diff --git a/Engine/Source/Programs/UnrealLightmass/UnrealLightmass.Build.cs b/Engine/Source/Programs/UnrealLightmass/UnrealLightmass.Build.cs
--- a/Engine/Source/Programs/UnrealLightmass/UnrealLightmass.Build.cs
+++ b/Engine/Source/Programs/UnrealLightmass/UnrealLightmass.Build.cs
@@ -28,6 +28,12 @@ public class UnrealLightmass : ModuleRules
 					"UdpMessaging",
 				}
 			);
+
+			PrivateIncludePathModuleNames.AddRange(
+				new string[] {
+					"Messaging",
+				}
+			);
 		}
 
 		// Lightmass ray tracing is 8% faster with buffer security checks disabled due to fixed size arrays on the stack in the kDOP ray tracing functions
