Baked skylighting
* Static sky lights are now allowed, baked into lightmap
* Stationary sky lights are baked into low quality lightmap
* Baked sky lighting can only support emissive input to avoid a feedback loop
* Adaptive refinement quadtree nodes are now pooled as Lightmass builds were 7x slower with the FBinnedAllocator otherwise
* Specializations of SHBasisFunction which saves 10% of lightmass build times with sky lighting
* Dark texels have their directionality's contribution to the lightmap min and max dampened

diff --git a/Engine/Shaders/BasePassPixelShader.usf b/Engine/Shaders/BasePassPixelShader.usf
--- a/Engine/Shaders/BasePassPixelShader.usf
+++ b/Engine/Shaders/BasePassPixelShader.usf
@@ -144,6 +144,7 @@ float3 GetTranslucencyLighting(FMaterialPixelParameters MaterialParameters, floa
 			
 		// Compute diffuse lighting which takes the normal into account
 		FTwoBandSHVector DiffuseTransferSH = CalcDiffuseTransferSH(MaterialParameters.WorldNormal, 1);
+		//@todo - why is x2 needed to match opaque and non-directional?
 		float4 VolumeLighting = float4(max(half3(0,0,0), DotSH(TranslucentLighting, DiffuseTransferSH)) * 2, Vector0.a);
 		InterpolatedLighting = DiffuseColor * VolumeLighting.rgb;
 
diff --git a/Engine/Shaders/ReflectionEnvironmentShaders.usf b/Engine/Shaders/ReflectionEnvironmentShaders.usf
--- a/Engine/Shaders/ReflectionEnvironmentShaders.usf
+++ b/Engine/Shaders/ReflectionEnvironmentShaders.usf
@@ -286,7 +286,7 @@ void DiffuseIrradianceCopyPS(
 
 	float SquaredUVs = 1 + dot(ScaledUVs, ScaledUVs);
 	// Dividing by NumSamples here to keep the sum in the range of fp16, once we get down to the 1x1 mip
-	float TexelWeight = 4 / (sqrt(SquaredUVs) * SquaredUVs * NumSamples);
+	float TexelWeight = 4 / (sqrt(SquaredUVs) * SquaredUVs);
 
 	FThreeBandSHVector SHCoefficients = SHBasisFunction3(CubeCoordinates);
 	float CurrentSHCoefficient = dot(SHCoefficients.V0, CoefficientMask0) + dot(SHCoefficients.V1, CoefficientMask1) + SHCoefficients.V2 * CoefficientMask2;
@@ -329,7 +329,7 @@ void DiffuseIrradianceAccumulatePS(
 		AccumulatedValue += SampleCubemap(CubeCoordinates);
 	}
 	
-	OutColor = AccumulatedValue;
+	OutColor = float4(AccumulatedValue.rgb / 4.0f, AccumulatedValue.a / 4.0f);
 }
 
 void AccumulateCubeFacesPS(
diff --git a/Engine/Source/Editor/UnrealEd/Private/Lightmass/Lightmass.cpp b/Engine/Source/Editor/UnrealEd/Private/Lightmass/Lightmass.cpp
--- a/Engine/Source/Editor/UnrealEd/Private/Lightmass/Lightmass.cpp
+++ b/Engine/Source/Editor/UnrealEd/Private/Lightmass/Lightmass.cpp
@@ -936,9 +936,14 @@ void FLightmassExporter::WriteLights( int32 Channel )
 	for ( int32 LightIndex = 0; LightIndex < SkyLights.Num(); ++LightIndex )
 	{
 		const USkyLightComponent* Light = SkyLights[LightIndex];
+
 		Lightmass::FLightData LightData;
 		Lightmass::FSkyLightData SkyData;
 		Copy( Light, LightData ); 
+
+		// Capture the scene's emissive and send it to lightmass
+		Light->CaptureEmissiveIrradianceEnvironmentMap(SkyData.IrradianceEnvironmentMap);
+
 		Swarm.WriteChannel( Channel, &LightData, sizeof(LightData) );
 		Swarm.WriteChannel( Channel, &SkyData, sizeof(SkyData) );
 		UpdateExportProgress();
diff --git a/Engine/Source/Programs/UnrealLightmass/Private/Lighting/FinalGather.cpp b/Engine/Source/Programs/UnrealLightmass/Private/Lighting/FinalGather.cpp
--- a/Engine/Source/Programs/UnrealLightmass/Private/Lighting/FinalGather.cpp
+++ b/Engine/Source/Programs/UnrealLightmass/Private/Lighting/FinalGather.cpp
@@ -150,6 +150,33 @@ FLinearColor FStaticLightingSystem::EvaluateEnvironmentLighting(const FVector4&
 	return IncomingDirection.Z < 0 ? (MaterialSettings.EnvironmentColor / (float)PI) : FLinearColor::Black;
 }
 
+void FStaticLightingSystem::EvaluateSkyLighting(const FVector4& IncomingDirection, bool bShadowed, FLinearColor& OutStaticLighting, FLinearColor& OutStationaryLighting) const
+{
+	for (int32 LightIndex = 0; LightIndex < SkyLights.Num(); LightIndex++)
+	{
+		FSkyLight* SkyLight = SkyLights[LightIndex];
+
+		if (!bShadowed || !(SkyLight->LightFlags & GI_LIGHT_CASTSHADOWS))
+		{
+			FSHVector3 SH = FSHVector3::SHBasisFunction(IncomingDirection);
+			FLinearColor Lighting = Dot(SkyLight->IrradianceEnvironmentMap, SH) * SkyLight->Brightness * FLinearColor(SkyLight->Color);
+
+			Lighting.R = FMath::Max(Lighting.R, 0.0f);
+			Lighting.G = FMath::Max(Lighting.G, 0.0f);
+			Lighting.B = FMath::Max(Lighting.B, 0.0f);
+
+			if (SkyLight->LightFlags & GI_LIGHT_HASSTATICLIGHTING)
+			{
+				OutStaticLighting += Lighting;
+			}
+			else if (SkyLight->LightFlags & GI_LIGHT_HASSTATICSHADOWING)
+			{
+				OutStationaryLighting += Lighting;
+			}
+		}
+	}
+}
+
 /** Calculates exitant radiance at a vertex. */
 FLinearColor FStaticLightingSystem::CalculateExitantRadiance(
 	const FStaticLightingMapping* SourceMapping,
@@ -401,14 +428,16 @@ FLinearColor FStaticLightingSystem::FinalGatherSample(
 	const FVector4& TriangleTangentPathDirection,
 	float SampleRadius,
 	int32 BounceNumber,
+	bool bDebugThisTexel,
 	FStaticLightingMappingContext& MappingContext,
 	FLMRandomStream& RandomStream,
 	FLightingCacheGatherInfo& RecordGatherInfo,
 	FFinalGatherInfo& FinalGatherInfo,
-	bool bDebugThisTexel,
-	FVector& OutUnoccludedSkyVector) const
+	FVector& OutUnoccludedSkyVector,
+	FLinearColor& OutStationarySkyLighting) const
 {
 	FLinearColor Lighting = FLinearColor::Black;
+	OutStationarySkyLighting = FLinearColor::Black;
 
 	checkSlow(TriangleTangentPathDirection.Z >= 0.0f);
 	checkSlow(TriangleTangentPathDirection.IsUnit3());
@@ -521,10 +550,10 @@ FLinearColor FStaticLightingSystem::FinalGatherSample(
 			const FLinearColor EnvironmentLighting = EvaluateEnvironmentLighting(-WorldPathDirection);
 			Lighting = EnvironmentLighting;
 		}
-
-		FinalGatherInfo.CombinedSkyUnoccludedDirections += WorldPathDirection;
 	}
 
+	EvaluateSkyLighting(WorldPathDirection, RayIntersection.bIntersects, Lighting, OutStationarySkyLighting);
+
 #if ALLOW_LIGHTMAP_SAMPLE_DEBUGGING
 	if (bDebugThisTexel
 		&& GeneralSettings.ViewSingleBounceNumber == BounceNumber
@@ -542,129 +571,6 @@ FLinearColor FStaticLightingSystem::FinalGatherSample(
 	return Lighting;
 }
 
-template<typename ElementType>
-class FSimpleQuadTreeNode
-{
-public:
-	ElementType Element;
-
-	FSimpleQuadTreeNode* Children[4];
-
-	FSimpleQuadTreeNode()
-	{
-		Children[0] = NULL;
-		Children[1] = NULL;
-		Children[2] = NULL;
-		Children[3] = NULL;
-	}
-
-	void AddChild(int32 Index, const ElementType& ChildElement)
-	{
-		Children[Index] = new FSimpleQuadTreeNode();
-		Children[Index]->Element = ChildElement;
-	}
-};
-
-template<typename ElementType>
-class FSimpleQuadTree
-{
-public:
-
-	const ElementType& GetLeafElement(float U, float V) const
-	{
-		return GetLeafElementRecursive(U, V, &RootNode);
-	}
-
-	FSimpleQuadTreeNode<ElementType> RootNode;
-
-	~FSimpleQuadTree()
-	{
-		DestroyTreeRecursive(&RootNode);
-	}
-
-private:
-
-	const ElementType& GetLeafElementRecursive(float U, float V, const FSimpleQuadTreeNode<ElementType>* Parent) const
-	{
-		const int32 ChildX = U > .5f;
-		const int32 ChildY = V > .5f;
-		const int32 ChildIndex = ChildX * 2 + ChildY;
-
-		if (Parent->Children[ChildIndex])
-		{
-			const float ChildU = U * 2 - ChildX;
-			const float ChildV = V * 2 - ChildY;
-
-			return GetLeafElementRecursive(ChildU, ChildV, Parent->Children[ChildIndex]);
-		}
-		else
-		{
-			return Parent->Element;
-		}
-	}
-
-	void DestroyTreeRecursive(const FSimpleQuadTreeNode<ElementType>* Node) const
-	{
-		for (int32 ChildIndex = 0; ChildIndex < ARRAY_COUNT(Node->Children); ChildIndex++)
-		{
-			if (Node->Children[ChildIndex])
-			{
-				DestroyTreeRecursive(Node->Children[ChildIndex]);
-			}
-		}
-		
-		if (Node != &RootNode)
-		{
-			delete Node;
-		}
-	}
-};
-
-class FLightingAndOcclusion
-{
-public:
-
-	FLinearColor Lighting;
-	FVector UnoccludedSkyVector;
-
-	FLightingAndOcclusion() :
-		Lighting(ForceInit),
-		UnoccludedSkyVector(FVector(0))
-	{}
-
-	FLightingAndOcclusion(const FLinearColor& InLighting, FVector InUnoccludedSkyVector) : 
-		Lighting(InLighting),
-		UnoccludedSkyVector(InUnoccludedSkyVector)
-	{}
-
-	friend inline FLightingAndOcclusion operator+ (const FLightingAndOcclusion& A, const FLightingAndOcclusion& B)
-	{
-		return FLightingAndOcclusion(A.Lighting + B.Lighting, A.UnoccludedSkyVector + B.UnoccludedSkyVector);
-	}
-
-	friend inline FLightingAndOcclusion operator/ (const FLightingAndOcclusion& A, float Divisor)
-	{
-		return FLightingAndOcclusion(A.Lighting / Divisor, A.UnoccludedSkyVector / Divisor);
-	}
-};
-
-/** Data stored for a sample that may need to be refined. */
-class FRefinementElement
-{
-public:
-	FLightingAndOcclusion Lighting;
-	FVector2D Uniforms;
-
-	FRefinementElement() :
-		Uniforms(FVector2D(0, 0))
-	{}
-
-	FRefinementElement(FLightingAndOcclusion InLighting, FVector2D InUniforms) :
-		Lighting(InLighting),
-		Uniforms(InUniforms)
-	{}
-};
-
 /** Stores intermediate data during a traversal of the refinement tree. */
 class FRefinementTraversalContext
 {
@@ -728,6 +634,14 @@ public:
 		Cells[ThetaIndex * NumPhiSteps + PhiIndex].RootNode.Element = Element;
 	}
 
+	void ReturnToFreeList(TArray<FSimpleQuadTreeNode<FRefinementElement>*>& OutNodes)
+	{
+		for (int32 CellIndex = 0; CellIndex < Cells.Num(); CellIndex++)
+		{
+			Cells[CellIndex].ReturnToFreeList(OutNodes);
+		}
+	}
+
 	void RefineIncomingRadiance(
 		const FStaticLightingSystem& LightingSystem,
 		const FStaticLightingMapping* Mapping,
@@ -787,7 +701,8 @@ public:
 			{
 				for (int32 PhiIndex = 0; PhiIndex < NumPhiSteps; PhiIndex++)
 				{
-					TotalLighting += GetFilteredValue(ThetaIndex, PhiIndex).Lighting;
+					const FLightingAndOcclusion FilteredLighting = GetFilteredValue(ThetaIndex, PhiIndex);
+					TotalLighting += FilteredLighting.Lighting + FilteredLighting.StationarySkyLighting;
 				}
 			}
 
@@ -830,7 +745,7 @@ public:
 						if (!bSuperSampleDueToImportanceCones)
 						{
 							const FLightingAndOcclusion RootElementLighting = GetRootValue(ThetaIndex, PhiIndex);
-							const FLinearColor Radiance = RootElementLighting.Lighting;
+							const FLinearColor Radiance = RootElementLighting.Lighting + RootElementLighting.StationarySkyLighting;
 							const float RelativeBrightness = Radiance.ComputeLuminance() / TotalBrightness;
 
 							for (int32 NeighborIndex = 0; NeighborIndex < ARRAY_COUNT(Neighbors); NeighborIndex++)
@@ -843,7 +758,7 @@ public:
 								if (NeighborTheta >= 0 && NeighborTheta < NumThetaSteps)
 								{
 									const FLightingAndOcclusion NeighborLighting = GetRootValue(NeighborTheta, NeighborPhi);
-									const float NeighborBrightness = NeighborLighting.Lighting.ComputeLuminance();
+									const float NeighborBrightness = (NeighborLighting.Lighting + NeighborLighting.StationarySkyLighting).ComputeLuminance();
 									const float NeighborRelativeBrightness = NeighborBrightness / TotalBrightness;
 									MaxRelativeDifference = FMath::Max(MaxRelativeDifference, FMath::Abs(RelativeBrightness - NeighborRelativeBrightness));
 									MaxSkyOcclusionDifference = FMath::Max(MaxSkyOcclusionDifference, FMath::Abs(RootElementLighting.UnoccludedSkyVector.SizeSquared() - NeighborLighting.UnoccludedSkyVector.SizeSquared()));
@@ -929,7 +844,7 @@ public:
 							if (!bSuperSampleDueToImportanceCones)
 							{
 								const FLightingAndOcclusion ChildLighting = ChildNode->Element.Lighting;
-								const FLinearColor Radiance = ChildLighting.Lighting;
+								const FLinearColor Radiance = ChildLighting.Lighting + ChildLighting.StationarySkyLighting;
 								const float RelativeBrightness = Radiance.ComputeLuminance() / TotalBrightness;
 
 								// Only search the axis neighbors past the first depth
@@ -941,7 +856,7 @@ public:
 									// Query must be done on the center of the cell
 									const FVector2D NeighborUV = FVector2D(NeighborU, NeighborV) + NodeContext.Size / 4;
 									const FLightingAndOcclusion NeighborLighting = GetValue(NeighborUV);
-									const float NeighborBrightness = NeighborLighting.Lighting.ComputeLuminance();
+									const float NeighborBrightness = (NeighborLighting.Lighting + NeighborLighting.StationarySkyLighting).ComputeLuminance();
 									const float NeighborRelativeBrightness = NeighborBrightness / TotalBrightness;
 									MaxRelativeDifference = FMath::Max(MaxRelativeDifference, FMath::Abs(RelativeBrightness - NeighborRelativeBrightness));
 									MaxSkyOcclusionDifference = FMath::Max(MaxSkyOcclusionDifference, FMath::Abs(ChildLighting.UnoccludedSkyVector.SizeSquared() - NeighborLighting.UnoccludedSkyVector.SizeSquared()));
@@ -987,7 +902,17 @@ public:
 				{
 					for (int32 SubPhiIndex = 0; SubPhiIndex < NumSubsamples; SubPhiIndex++)
 					{
-						FRefinementElement SubsampleElement;
+						FSimpleQuadTreeNode<FRefinementElement>* FreeNode = NULL;
+						
+						if (MappingContext.RefinementTreeFreePool.Num() > 0)
+						{
+							FreeNode = MappingContext.RefinementTreeFreePool.Pop(false);
+							*FreeNode = FSimpleQuadTreeNode<FRefinementElement>();
+						}
+						else
+						{
+							FreeNode = new FSimpleQuadTreeNode<FRefinementElement>();
+						}
 
 						const FVector2D ChildMin = NodeContext.Min + FVector2D(SubThetaIndex, SubPhiIndex) * NodeContext.Size / 2;
 
@@ -997,7 +922,7 @@ public:
 							&& NodeContext.Node->Element.Uniforms.X < ChildMin.X + NodeContext.Size.X / 2
 							&& NodeContext.Node->Element.Uniforms.Y < ChildMin.Y + NodeContext.Size.Y / 2)
 						{
-							SubsampleElement = NodeContext.Node->Element;
+							FreeNode->Element = NodeContext.Node->Element;
 						}
 						else
 						{
@@ -1012,20 +937,23 @@ public:
 							const FVector4 SampleDirection = UniformSampleHemisphere(Fraction1, Fraction2);
 
 							FVector UnoccludedSkyVector;
+							FLinearColor StationarySkyLighting;
+
 							const FLinearColor SubsampleLighting = LightingSystem.FinalGatherSample(
 								Mapping,
 								Vertex,
 								SampleDirection,
 								SampleRadius,
 								BounceNumber,
+								bDebugThisTexel,
 								MappingContext,
 								RandomStream,
 								SubsampleGatherInfo,
 								SubsampleFinalGatherInfo,
-								bDebugThisTexel,
-								UnoccludedSkyVector);
+								UnoccludedSkyVector,
+								StationarySkyLighting);
 
-							SubsampleElement = FRefinementElement(FLightingAndOcclusion(SubsampleLighting, UnoccludedSkyVector), FVector2D(Fraction1, Fraction2));
+							FreeNode->Element = FRefinementElement(FLightingAndOcclusion(SubsampleLighting, UnoccludedSkyVector, StationarySkyLighting), FVector2D(Fraction1, Fraction2));
 
 							Stats.NumRefiningFinalGatherSamples[RefinementDepth]++;
 
@@ -1039,7 +967,7 @@ public:
 							}
 						}
 
-						NodeContext.Node->AddChild(SubThetaIndex * NumSubsamples + SubPhiIndex, SubsampleElement);
+						NodeContext.Node->AddChild(SubThetaIndex * NumSubsamples + SubPhiIndex, FreeNode);
 					}
 				}
 			}
@@ -1049,10 +977,7 @@ public:
 			ImportanceConeAngle /= 4;
 			BrightnessThreshold *= 2;
 			ConeWeightThreshold *= 1.5f;
-
-			// Currently limiting sky occlusion refinement to only one level to keep build times down
-			//@todo - this will cap sky occlusion quality, allow scaling up
-			SkyOcclusionThreshold *= 15;
+			SkyOcclusionThreshold *= 16;
 		}
 	}
 
@@ -1132,20 +1057,23 @@ SampleType FStaticLightingSystem::IncomingRadianceAdaptive(
 			const FVector4 TriangleTangentPathDirection = UniformHemisphereSamples[SampleIndex];
 			
 			FVector UnoccludedSkyVector;
+			FLinearColor StationarySkyLighting;
+
 			const FLinearColor Radiance = FinalGatherSample(
 				Mapping,
 				Vertex,
 				TriangleTangentPathDirection,
 				SampleRadius,
 				BounceNumber,
+				bDebugThisTexel,
 				MappingContext,
 				RandomStream,
 				GatherInfo,
 				FinalGatherInfo,
-				bDebugThisTexel,
-				UnoccludedSkyVector);
+				UnoccludedSkyVector,
+				StationarySkyLighting);
 
-			RefinementGrid.SetRootElement(ThetaIndex, PhiIndex, FRefinementElement(FLightingAndOcclusion(Radiance, UnoccludedSkyVector), UniformHemisphereSampleUniforms[SampleIndex]));
+			RefinementGrid.SetRootElement(ThetaIndex, PhiIndex, FRefinementElement(FLightingAndOcclusion(Radiance, UnoccludedSkyVector, StationarySkyLighting), UniformHemisphereSampleUniforms[SampleIndex]));
 		}
 	}
 
@@ -1213,6 +1141,7 @@ SampleType FStaticLightingSystem::IncomingRadianceAdaptive(
 			CombinedSkyUnoccludedDirection += FilteredLighting.UnoccludedSkyVector;
 
 			IncomingRadiance.AddIncomingRadiance(Radiance, SampleWeight, TangentPathDirection, WorldPathDirection);
+			IncomingRadiance.AddIncomingStationarySkyLight(FilteredLighting.StationarySkyLighting, SampleWeight, TangentPathDirection, WorldPathDirection);
 			checkSlow(IncomingRadiance.AreFloatsValid());
 		}
 	}
@@ -1230,6 +1159,8 @@ SampleType FStaticLightingSystem::IncomingRadianceAdaptive(
 	const FVector BentNormal = CombinedSkyUnoccludedDirection / (MaxUnoccludedLength * UniformHemisphereSamples.Num());
 	IncomingRadiance.SetSkyOcclusion(BentNormal);
 
+	RefinementGrid.ReturnToFreeList(MappingContext.RefinementTreeFreePool);
+
 	return IncomingRadiance;
 }
 
diff --git a/Engine/Source/Programs/UnrealLightmass/Private/Lighting/LightingSystem.cpp b/Engine/Source/Programs/UnrealLightmass/Private/Lighting/LightingSystem.cpp
--- a/Engine/Source/Programs/UnrealLightmass/Private/Lighting/LightingSystem.cpp
+++ b/Engine/Source/Programs/UnrealLightmass/Private/Lighting/LightingSystem.cpp
@@ -25,8 +25,6 @@
 
 namespace Lightmass
 {
-
-
 	 
 FGatheredLightSample FGatheredLightSample::AmbientLight(const FLinearColor& Color)
 {
@@ -36,10 +34,11 @@ FGatheredLightSample FGatheredLightSample::AmbientLight(const FLinearColor& Colo
 	// Compute SHCorrection as if all the lighting was coming in along the normal
 	FVector4 TangentDirection(0, 0, 1);
 	FSHVector2 SH = FSHVector2::SHBasisFunction( TangentDirection );
-	Result.SHCorrection += Color.GetLuminance() * ( 0.282095f * SH.V[0] + 0.325735f * SH.V[2] );
-
+	Result.SHCorrection = Color.GetLuminance() * ( 0.282095f * SH.V[0] + 0.325735f * SH.V[2] );
 	Result.IncidentLighting = Color;
 
+	checkSlow(Result.SHCorrection >= 0 && Result.IncidentLighting.GetMin() >= 0);
+
 	return Result;
 }
 
@@ -56,8 +55,10 @@ FGatheredLightSample FGatheredLightSample::PointLightWorldSpace(const FLinearCol
 		// These scaling coefficients are SHBasisFunction and CalcDiffuseTransferSH baked down
 		// 0.325735f = 0.488603f from SHBasisFunction * 2/3 from CalcDiffuseTransferSH
 		// Only using V[2] which is the tangent space Z
-		Result.SHCorrection += Color.GetLuminance() * ( 0.282095f * SH.V[0] + 0.325735f * SH.V[2] );
+		Result.SHCorrection = Color.GetLuminance() * ( 0.282095f * SH.V[0] + 0.325735f * SH.V[2] );
 		Result.IncidentLighting = Color * FMath::Max(0.0f, TangentDirection.Z);
+
+		checkSlow(Result.SHCorrection >= 0 && Result.IncidentLighting.GetMin() >= 0);
 	}
 
 	return Result;
@@ -153,6 +154,7 @@ void FFinalGatherSample::AddWeighted(const FFinalGatherSample& OtherSample, floa
 {
 	FGatheredLightSample::AddWeighted(OtherSample, Weight);
 	Occlusion += OtherSample.Occlusion * Weight;
+	StationarySkyLighting = StationarySkyLighting + OtherSample.StationarySkyLighting * Weight;
 }
 
 bool FFinalGatherSample::AreFloatsValid() const
@@ -167,14 +169,21 @@ FStaticLightingMappingContext::FStaticLightingMappingContext(const FStaticLighti
 
 FStaticLightingMappingContext::~FStaticLightingMappingContext()
 {
-	// Update the main threads stats with the stats from this mapping
-	FScopeLock Lock(&System.Stats.StatsSync);
-	System.Stats.Cache[0] += FirstBounceCache.Stats;
-	System.Stats += Stats;
-	System.Stats.NumFirstHitRaysTraced += RayCache.NumFirstHitRaysTraced;
-	System.Stats.NumBooleanRaysTraced += RayCache.NumBooleanRaysTraced;
-	System.Stats.FirstHitRayTraceThreadTime += RayCache.FirstHitRayTraceTime;
-	System.Stats.BooleanRayTraceThreadTime += RayCache.BooleanRayTraceTime;
+	{
+		// Update the main threads stats with the stats from this mapping
+		FScopeLock Lock(&System.Stats.StatsSync);
+		System.Stats.Cache[0] += FirstBounceCache.Stats;
+		System.Stats += Stats;
+		System.Stats.NumFirstHitRaysTraced += RayCache.NumFirstHitRaysTraced;
+		System.Stats.NumBooleanRaysTraced += RayCache.NumBooleanRaysTraced;
+		System.Stats.FirstHitRayTraceThreadTime += RayCache.FirstHitRayTraceTime;
+		System.Stats.BooleanRayTraceThreadTime += RayCache.BooleanRayTraceTime;
+	}
+
+	for (int32 EntryIndex = 0; EntryIndex < RefinementTreeFreePool.Num(); EntryIndex++)
+	{
+		delete RefinementTreeFreePool[EntryIndex];
+	}
 }
 
 /**
diff --git a/Engine/Source/Programs/UnrealLightmass/Private/Lighting/LightingSystem.h b/Engine/Source/Programs/UnrealLightmass/Private/Lighting/LightingSystem.h
--- a/Engine/Source/Programs/UnrealLightmass/Private/Lighting/LightingSystem.h
+++ b/Engine/Source/Programs/UnrealLightmass/Private/Lighting/LightingSystem.h
@@ -74,14 +74,6 @@ public:
 	 */
 	static FGatheredLightSample PointLightWorldSpace(const FLinearColor& Color, const FVector4& TangentDirection, const FVector4& WorldDirection);
 
-	/**
-	 * Constructs a light sample representing a sky light.
-	 * @param UpperColor - The color/intensity of the sky light's upper hemisphere.
-	 * @param LowerColor - The color/intensity of the sky light's lower hemisphere.
-	 * @param WorldZ - The world's +Z axis in tangent space.
-	 */
-	static FGatheredLightSample SkyLight(const FLinearColor& UpperColor,const FLinearColor& LowerColor,const FVector4& WorldZ);
-
 	/**
 	 * Adds a weighted light sample to this light sample.
 	 * @param OtherSample - The sample to add.
@@ -224,6 +216,11 @@ public:
 	/** Occlusion factor of the sample, 0 is completely unoccluded, 1 is completely occluded. */
 	float Occlusion;
 
+	/** 
+	 * A light sample for sky lighting.  This has to be stored separately to support stationary sky lights only contributing to low quality lightmaps.
+	 */
+	FGatheredLightSample StationarySkyLighting;
+
 	/** Initialization constructor. */
 	FFinalGatherSample() : 
 		FGatheredLightSample(),
@@ -262,6 +259,11 @@ public:
 		AddWeighted(FGatheredLightSample::PointLightWorldSpace(IncomingRadiance, TangentSpaceDirection, WorldSpaceDirection), Weight);
 	}
 
+	inline void AddIncomingStationarySkyLight(const FLinearColor& IncomingSkyLight, float Weight, const FVector4& TangentSpaceDirection, const FVector4& WorldSpaceDirection)
+	{
+		StationarySkyLighting.AddWeighted(FGatheredLightSample::PointLightWorldSpace(IncomingSkyLight, TangentSpaceDirection, WorldSpaceDirection), Weight);
+	}
+
 	bool AreFloatsValid() const;
 
 	FFinalGatherSample operator*(float Scalar) const
@@ -269,6 +271,7 @@ public:
 		FFinalGatherSample Result;
 		(FGatheredLightSample&)Result = (const FGatheredLightSample&)(*this) * Scalar;
 		Result.Occlusion = Occlusion * Scalar;
+		Result.StationarySkyLighting = StationarySkyLighting * Scalar;
 		return Result;
 	}
 
@@ -277,6 +280,7 @@ public:
 		FFinalGatherSample Result;
 		(FGatheredLightSample&)Result = (const FGatheredLightSample&)(*this) + (const FGatheredLightSample&)SampleB;
 		Result.Occlusion = Occlusion + SampleB.Occlusion;
+		Result.StationarySkyLighting = StationarySkyLighting + SampleB.StationarySkyLighting;
 		return Result;
 	}
 };
@@ -285,13 +289,11 @@ struct FFinalGatherInfo
 {
 	FFinalGatherInfo() :
 		NumBackfaceHits(0),
-		NumSamplesOccluded(0),
-		CombinedSkyUnoccludedDirections(FVector(0))
+		NumSamplesOccluded(0)
 	{}
 
 	int32 NumBackfaceHits;
 	float NumSamplesOccluded;
-	FVector CombinedSkyUnoccludedDirections;
 };
 
 struct FTexelCorner
@@ -1244,6 +1246,136 @@ public:
 	}
 };
 
+template<typename ElementType>
+class FSimpleQuadTreeNode
+{
+public:
+	ElementType Element;
+
+	FSimpleQuadTreeNode* Children[4];
+
+	FSimpleQuadTreeNode()
+	{
+		Children[0] = NULL;
+		Children[1] = NULL;
+		Children[2] = NULL;
+		Children[3] = NULL;
+	}
+
+	void AddChild(int32 Index, FSimpleQuadTreeNode* Child)
+	{
+		Children[Index] = Child;
+	}
+};
+
+template<typename ElementType>
+class FSimpleQuadTree
+{
+public:
+
+	const ElementType& GetLeafElement(float U, float V) const
+	{
+		return GetLeafElementRecursive(U, V, &RootNode);
+	}
+
+	FSimpleQuadTreeNode<ElementType> RootNode;
+
+	~FSimpleQuadTree()
+	{
+	}
+
+	void ReturnToFreeList(TArray<FSimpleQuadTreeNode<ElementType>*>& OutNodes)
+	{
+		ReturnToFreeListRecursive(&RootNode, OutNodes);
+	}
+
+private:
+
+	const ElementType& GetLeafElementRecursive(float U, float V, const FSimpleQuadTreeNode<ElementType>* Parent) const
+	{
+		const int32 ChildX = U > .5f;
+		const int32 ChildY = V > .5f;
+		const int32 ChildIndex = ChildX * 2 + ChildY;
+
+		if (Parent->Children[ChildIndex])
+		{
+			const float ChildU = U * 2 - ChildX;
+			const float ChildV = V * 2 - ChildY;
+
+			return GetLeafElementRecursive(ChildU, ChildV, Parent->Children[ChildIndex]);
+		}
+		else
+		{
+			return Parent->Element;
+		}
+	}
+
+	void ReturnToFreeListRecursive(FSimpleQuadTreeNode<ElementType>* Node, TArray<FSimpleQuadTreeNode<ElementType>*>& OutNodes) const
+	{
+		for (int32 ChildIndex = 0; ChildIndex < ARRAY_COUNT(Node->Children); ChildIndex++)
+		{
+			if (Node->Children[ChildIndex])
+			{
+				ReturnToFreeListRecursive(Node->Children[ChildIndex], OutNodes);
+			}
+		}
+
+		if (Node != &RootNode)
+		{
+			OutNodes.Add(Node);
+		}
+	}
+};
+
+/** Lighting payload used by the adaptive refinement in final gathering. */
+class FLightingAndOcclusion
+{
+public:
+
+	FLinearColor Lighting;
+	FVector UnoccludedSkyVector;
+	FLinearColor StationarySkyLighting;
+
+	FLightingAndOcclusion() :
+		Lighting(ForceInit),
+		UnoccludedSkyVector(FVector(0)),
+		StationarySkyLighting(FLinearColor::Black)
+	{}
+
+	FLightingAndOcclusion(const FLinearColor& InLighting, FVector InUnoccludedSkyVector, const FLinearColor& InStationarySkyLighting) : 
+		Lighting(InLighting),
+		UnoccludedSkyVector(InUnoccludedSkyVector),
+		StationarySkyLighting(InStationarySkyLighting)
+	{}
+
+	friend inline FLightingAndOcclusion operator+ (const FLightingAndOcclusion& A, const FLightingAndOcclusion& B)
+	{
+		return FLightingAndOcclusion(A.Lighting + B.Lighting, A.UnoccludedSkyVector + B.UnoccludedSkyVector, A.StationarySkyLighting + B.StationarySkyLighting);
+	}
+
+	friend inline FLightingAndOcclusion operator/ (const FLightingAndOcclusion& A, float Divisor)
+	{
+		return FLightingAndOcclusion(A.Lighting / Divisor, A.UnoccludedSkyVector / Divisor, A.StationarySkyLighting / Divisor);
+	}
+};
+
+/** Data stored for a sample that may need to be refined. */
+class FRefinementElement
+{
+public:
+	FLightingAndOcclusion Lighting;
+	FVector2D Uniforms;
+
+	FRefinementElement() :
+		Uniforms(FVector2D(0, 0))
+	{}
+
+	FRefinementElement(FLightingAndOcclusion InLighting, FVector2D InUniforms) :
+		Lighting(InLighting),
+		Uniforms(InUniforms)
+	{}
+};
+
 /** Local state for a mapping, accessed only by the owning thread. */
 class FStaticLightingMappingContext
 {
@@ -1257,6 +1389,8 @@ public:
 
 	TArray<FDebugLightingCacheRecord> DebugCacheRecords;
 
+	TArray<FSimpleQuadTreeNode<FRefinementElement>* > RefinementTreeFreePool;
+
 	class FStaticLightingSystem& System;
 
 public:
@@ -2116,6 +2250,9 @@ private:
 	/** Returns environment lighting for the given direction. */
 	FLinearColor EvaluateEnvironmentLighting(const FVector4& IncomingDirection) const;
 
+	/** Evaluates the incoming sky lighting from the scene. */
+	void EvaluateSkyLighting(const FVector4& IncomingDirection, bool bShadowed, FLinearColor& OutStaticLighting, FLinearColor& OutStationaryLighting) const;
+
 	/** Returns a light sample that represents the material attribute specified by MaterialSettings.ViewMaterialAttribute at the intersection. */
 	FGatheredLightSample GetVisualizedMaterialAttribute(const FStaticLightingMapping* Mapping, const FLightRayIntersection& Intersection) const;
 
@@ -2153,12 +2290,13 @@ private:
 		const FVector4& TriangleTangentPathDirection,
 		float SampleRadius,
 		int32 BounceNumber,
+		bool bDebugThisTexel,
 		FStaticLightingMappingContext& MappingContext,
 		FLMRandomStream& RandomStream,
 		FLightingCacheGatherInfo& GatherInfo,
 		FFinalGatherInfo& FinalGatherInfo,
-		bool bDebugThisTexel,
-		FVector& OutUnoccludedSkyVector) const;
+		FVector& OutUnoccludedSkyVector,
+		FLinearColor& OutStationarySkyLighting) const;
 
 	/** 
 	 * Final gather using adaptive sampling to estimate the incident radiance function. 
@@ -2370,7 +2508,7 @@ private:
 	/** The lights in the world which the system is building, excluding sky lights. */
 	TArray<FLight*> Lights;
 
-	TArray<FLight*> SkyLights;
+	TArray<FSkyLight*> SkyLights;
 
 	/** Mesh area lights in the world. */
 	TIndirectArray<FMeshAreaLight> MeshAreaLights;
diff --git a/Engine/Source/Programs/UnrealLightmass/Private/Lighting/LightmapData.cpp b/Engine/Source/Programs/UnrealLightmass/Private/Lighting/LightmapData.cpp
--- a/Engine/Source/Programs/UnrealLightmass/Private/Lighting/LightmapData.cpp
+++ b/Engine/Source/Programs/UnrealLightmass/Private/Lighting/LightmapData.cpp
@@ -164,14 +164,13 @@ namespace Lightmass
 					MinCoefficient[0][3] = FMath::Min( MinCoefficient[0][3], LogL );
 					MaxCoefficient[0][3] = FMath::Max( MaxCoefficient[0][3], LogL );
 
-					// Don't care about directionality for black
-					if( L > 1e-4 )
+					// Dampen dark texel's contribution on the directionality min and max
+					float DampenDirectionality = FMath::Clamp(L * 100, 0.0f, 1.0f);
+
+					for( int32 ColorIndex = 0; ColorIndex < 3; ColorIndex++ )
 					{
-						for( int32 ColorIndex = 0; ColorIndex < 3; ColorIndex++ )
-						{
-							MinCoefficient[1][ColorIndex] = FMath::Min( MinCoefficient[1][ColorIndex], SourceSample.Coefficients[1][ColorIndex] );
-							MaxCoefficient[1][ColorIndex] = FMath::Max( MaxCoefficient[1][ColorIndex], SourceSample.Coefficients[1][ColorIndex] );
-						}
+						MinCoefficient[1][ColorIndex] = FMath::Min( MinCoefficient[1][ColorIndex], DampenDirectionality * SourceSample.Coefficients[1][ColorIndex] );
+						MaxCoefficient[1][ColorIndex] = FMath::Max( MaxCoefficient[1][ColorIndex], DampenDirectionality * SourceSample.Coefficients[1][ColorIndex] );
 					}
 				}
 
@@ -195,14 +194,13 @@ namespace Lightmass
 					MinCoefficient[2][2] = FMath::Min( MinCoefficient[2][2], LogB );
 					MaxCoefficient[2][2] = FMath::Max( MaxCoefficient[2][2], LogB );
 
-					// Don't care about directionality for black
-					if( L > 1e-4 )
+					// Dampen dark texel's contribution on the directionality min and max
+					float DampenDirectionality = FMath::Clamp(L * 100, 0.0f, 1.0f);
+
+					for( int32 ColorIndex = 0; ColorIndex < 3; ColorIndex++ )
 					{
-						for( int32 ColorIndex = 0; ColorIndex < 3; ColorIndex++ )
-						{
-							MinCoefficient[3][ColorIndex] = FMath::Min( MinCoefficient[3][ColorIndex], SourceSample.Coefficients[3][ColorIndex] );
-							MaxCoefficient[3][ColorIndex] = FMath::Max( MaxCoefficient[3][ColorIndex], SourceSample.Coefficients[3][ColorIndex] );
-						}
+						MinCoefficient[3][ColorIndex] = FMath::Min( MinCoefficient[3][ColorIndex], DampenDirectionality * SourceSample.Coefficients[3][ColorIndex] );
+						MaxCoefficient[3][ColorIndex] = FMath::Max( MaxCoefficient[3][ColorIndex], DampenDirectionality * SourceSample.Coefficients[3][ColorIndex] );
 					}
 				}
 			}
diff --git a/Engine/Source/Programs/UnrealLightmass/Private/Lighting/TextureMapping.cpp b/Engine/Source/Programs/UnrealLightmass/Private/Lighting/TextureMapping.cpp
--- a/Engine/Source/Programs/UnrealLightmass/Private/Lighting/TextureMapping.cpp
+++ b/Engine/Source/Programs/UnrealLightmass/Private/Lighting/TextureMapping.cpp
@@ -2604,6 +2604,7 @@ void FStaticLightingSystem::ProcessInterpolateTask(FInterpolateIndirectTaskDescr
 
 				// Replace sky occlusion in the lighting sample that will be written into the lightmap with the interpolated sky occlusion using IrradianceCachingSettings.SkyOcclusionSmoothnessReduction
 				IndirectLighting.SkyOcclusion = SecondInterpolatedIndirectLighting.SkyOcclusion;
+				IndirectLighting.StationarySkyLighting = SecondInterpolatedIndirectLighting.StationarySkyLighting;
 
 				if (Task->TextureMapping->Mesh->UsesTwoSidedLighting(TexelToVertex.ElementIndex))
 				{
@@ -2635,6 +2636,9 @@ void FStaticLightingSystem::ProcessInterpolateTask(FInterpolateIndirectTaskDescr
 				// Apply occlusion to indirect lighting and add this texel's indirect lighting to its running total
 				CurrentLightSample.AddWeighted(IndirectLighting, 1);
 
+				// Stationary sky light contribution goes into low quality lightmap only, bent normal sky shadowing will be exported separately
+				CurrentLightSample.LowQuality.AddWeighted(IndirectLighting.StationarySkyLighting, 1);
+
 				if (AmbientOcclusionSettings.bUseAmbientOcclusion && AmbientOcclusionSettings.bVisualizeAmbientOcclusion)
 				{
 					//@todo - this will only be the correct intensity for simple lightmaps
diff --git a/Engine/Source/Programs/UnrealLightmass/Public/SceneExport.h b/Engine/Source/Programs/UnrealLightmass/Public/SceneExport.h
--- a/Engine/Source/Programs/UnrealLightmass/Public/SceneExport.h
+++ b/Engine/Source/Programs/UnrealLightmass/Public/SceneExport.h
@@ -871,6 +871,7 @@ struct FSpotLightData
 //----------------------------------------------------------------------------
 struct FSkyLightData
 {
+	FSHVectorRGB3 IrradianceEnvironmentMap;
 };
 
 //----------------------------------------------------------------------------
diff --git a/Engine/Source/Runtime/Core/Public/Math/SHMath.h b/Engine/Source/Runtime/Core/Public/Math/SHMath.h
--- a/Engine/Source/Runtime/Core/Public/Math/SHMath.h
+++ b/Engine/Source/Runtime/Core/Public/Math/SHMath.h
@@ -331,6 +331,35 @@ public:
 	}
 } GCC_ALIGN(16);
 
+/** Specialization for 2nd order to avoid expensive trig functions. */
+template<> TSHVector<2> TSHVector<2>::SHBasisFunction(const FVector& Vector) 
+{
+	TSHVector<2> Result;
+	Result.V[0] = 0.282095f; 
+	Result.V[1] = -0.488603f * Vector.Y;
+	Result.V[2] = 0.488603f * Vector.Z;
+	Result.V[3] = -0.488603f * Vector.X;
+	return Result;
+}
+
+/** Specialization for 3rd order to avoid expensive trig functions. */
+template<> TSHVector<3> TSHVector<3>::SHBasisFunction(const FVector& Vector) 
+{
+	TSHVector<3> Result;
+	Result.V[0] = 0.282095f; 
+	Result.V[1] = -0.488603f * Vector.Y;
+	Result.V[2] = 0.488603f * Vector.Z;
+	Result.V[3] = -0.488603f * Vector.X;
+
+	FVector VectorSquared = Vector * Vector;
+	Result.V[4] = 1.092548f * Vector.X * Vector.Y;
+	Result.V[5] = -1.092548f * Vector.Y * Vector.Z;
+	Result.V[6] = 0.315392f * (3.0f * VectorSquared.Z - 1.0f);
+	Result.V[7] = -1.092548f * Vector.X * Vector.Z;
+	Result.V[8] = 0.546274f * (VectorSquared.X - VectorSquared.Y);
+	return Result;
+}
+
 /** A vector of colored spherical harmonic coefficients. */
 template<int32 MaxSHOrder>
 class TSHVectorRGB
@@ -349,6 +378,15 @@ public:
 		return R * 0.3f + G * 0.59f + B * 0.11f;
 	}
 
+	void Desaturate(float DesaturateFraction)
+	{
+		TSHVector<MaxSHOrder> Desaturated = GetLuminance() * DesaturateFraction;
+
+		R = R * (1 - DesaturateFraction) + Desaturated;
+		G = G * (1 - DesaturateFraction) + Desaturated;
+		B = B * (1 - DesaturateFraction) + Desaturated;
+	}
+
 	/** Calculates the integral of the function over the surface of the sphere. */
 	FLinearColor CalcIntegral() const
 	{
diff --git a/Engine/Source/Runtime/Engine/Classes/Components/SkyLightComponent.h b/Engine/Source/Runtime/Engine/Classes/Components/SkyLightComponent.h
--- a/Engine/Source/Runtime/Engine/Classes/Components/SkyLightComponent.h
+++ b/Engine/Source/Runtime/Engine/Classes/Components/SkyLightComponent.h
@@ -128,6 +128,9 @@ class ENGINE_API USkyLightComponent : public ULightComponentBase
 	/** Called each tick to recapture and queued sky captures. */
 	static void UpdateSkyCaptureContents(UWorld* WorldToUpdate);
 
+	/** Computes an irradiance environment map using only emissive contribution from the sky light. */
+	void CaptureEmissiveIrradianceEnvironmentMap(FSHVectorRGB3& OutIrradianceMap) const;
+
 	UFUNCTION(BlueprintCallable, Category="Rendering|Components|SkyLight")
 	void SetBrightness(float NewBrightness);
 
diff --git a/Engine/Source/Runtime/Engine/Private/SkyLightComponent.cpp b/Engine/Source/Runtime/Engine/Private/SkyLightComponent.cpp
--- a/Engine/Source/Runtime/Engine/Private/SkyLightComponent.cpp
+++ b/Engine/Source/Runtime/Engine/Private/SkyLightComponent.cpp
@@ -71,6 +71,7 @@ FSkyLightSceneProxy::FSkyLightSceneProxy(const USkyLightComponent* InLightCompon
 	, SkyDistanceThreshold(InLightComponent->SkyDistanceThreshold)
 	, bCastShadows(InLightComponent->CastShadows)
 	, bPrecomputedLightingIsValid(InLightComponent->bPrecomputedLightingIsValid)
+	, bHasStaticLighting(InLightComponent->HasStaticLighting())
 	, LightColor(FLinearColor(InLightComponent->LightColor) * InLightComponent->Intensity)
 	, IrradianceEnvironmentMap(InLightComponent->IrradianceEnvironmentMap)
 {
@@ -196,8 +197,11 @@ void USkyLightComponent::PostEditChangeProperty(FPropertyChangedEvent& PropertyC
 {
 	const FName CategoryName = FObjectEditorUtils::GetCategoryFName(PropertyChangedEvent.Property);
 
-	// Other options not supported yet
-	Mobility = EComponentMobility::Stationary;
+	// Movable not supported yet
+	if (Mobility == EComponentMobility::Movable)
+	{
+		Mobility = EComponentMobility::Stationary;
+	}
 
 	Super::PostEditChangeProperty(PropertyChangedEvent);
 
@@ -374,7 +378,9 @@ void USkyLightComponent::UpdateSkyCaptureContents(UWorld* WorldToUpdate)
 						CaptureComponent->MarkRenderStateDirty();
 					}
 
-					WorldToUpdate->Scene->UpdateSkyCaptureContents(CaptureComponent);
+					WorldToUpdate->Scene->UpdateSkyCaptureContents(CaptureComponent, false, CaptureComponent->ProcessedSkyTexture, CaptureComponent->IrradianceEnvironmentMap);
+
+					CaptureComponent->MarkRenderStateDirty();
 				}
 
 				// Only remove queued update requests if we processed it for the right world
@@ -384,6 +390,18 @@ void USkyLightComponent::UpdateSkyCaptureContents(UWorld* WorldToUpdate)
 	}
 }
 
+void USkyLightComponent::CaptureEmissiveIrradianceEnvironmentMap(FSHVectorRGB3& OutIrradianceMap) const
+{
+	OutIrradianceMap = FSHVectorRGB3();
+
+	if (GetScene())
+	{
+		// Capture emissive scene lighting only for the lighting build
+		// This is necessary to avoid a feedback loop with the last lighting build results
+		GetScene()->UpdateSkyCaptureContents(this, true, NULL, OutIrradianceMap);
+	}
+}
+
 /** Set brightness of the light */
 void USkyLightComponent::SetBrightness(float NewBrightness)
 {
diff --git a/Engine/Source/Runtime/Engine/Public/SceneManagement.h b/Engine/Source/Runtime/Engine/Public/SceneManagement.h
--- a/Engine/Source/Runtime/Engine/Public/SceneManagement.h
+++ b/Engine/Source/Runtime/Engine/Public/SceneManagement.h
@@ -720,6 +720,7 @@ public:
 	float SkyDistanceThreshold;
 	bool bCastShadows;
 	bool bPrecomputedLightingIsValid;
+	bool bHasStaticLighting;
 	FLinearColor LightColor;
 	FSHVectorRGB3 IrradianceEnvironmentMap;
 };
@@ -1837,7 +1838,7 @@ public:
 	 * Updates the contents of the given sky capture by rendering the scene. 
 	 * This must be called on the game thread.
 	 */
-	virtual void UpdateSkyCaptureContents(USkyLightComponent* CaptureComponent) {}
+	virtual void UpdateSkyCaptureContents(const USkyLightComponent* CaptureComponent, bool bCaptureEmissiveOnly, FTexture* OutProcessedTexture, FSHVectorRGB3& OutIrradianceEnvironmentMap) {}
 
 	/** 
 	* Updates the contents of the given scene capture by rendering the scene. 
diff --git a/Engine/Source/Runtime/Renderer/Private/BasePassRendering.cpp b/Engine/Source/Runtime/Renderer/Private/BasePassRendering.cpp
--- a/Engine/Source/Runtime/Renderer/Private/BasePassRendering.cpp
+++ b/Engine/Source/Runtime/Renderer/Private/BasePassRendering.cpp
@@ -279,7 +279,7 @@ public:
 			LightMapPolicy,
 			Parameters.BlendMode,
 			Parameters.TextureMode,
-			Scene && Scene->SkyLight && bIsLitMaterial,
+			Scene && Scene->SkyLight && !Scene->SkyLight->bHasStaticLighting && bIsLitMaterial,
 			IsTranslucentBlendMode(Parameters.BlendMode) && (Scene && Scene->HasAtmosphericFog()) && View.Family->EngineShowFlags.Atmosphere,
 			View.Family->EngineShowFlags.ShaderComplexity,
 			false,
diff --git a/Engine/Source/Runtime/Renderer/Private/ReflectionEnvironmentCapture.cpp b/Engine/Source/Runtime/Renderer/Private/ReflectionEnvironmentCapture.cpp
--- a/Engine/Source/Runtime/Renderer/Private/ReflectionEnvironmentCapture.cpp
+++ b/Engine/Source/Runtime/Renderer/Private/ReflectionEnvironmentCapture.cpp
@@ -1310,7 +1310,7 @@ private:
 
 TGlobalResource<FCaptureRenderTarget> GReflectionCaptureRenderTarget;
 
-void CaptureSceneIntoScratchCubemap(FScene* Scene, FVector CapturePosition, bool bCapturingForSkyLight, float SkyLightNearPlane, bool bLowerHemisphereIsBlack)
+void CaptureSceneIntoScratchCubemap(FScene* Scene, FVector CapturePosition, bool bCapturingForSkyLight, float SkyLightNearPlane, bool bLowerHemisphereIsBlack, bool bCaptureEmissiveOnly)
 {
 	for (int32 CubeFace = 0; CubeFace < CubeFace_MAX; CubeFace++)
 	{
@@ -1369,6 +1369,12 @@ void CaptureSceneIntoScratchCubemap(FScene* Scene, FVector CapturePosition, bool
 		// Force all surfaces diffuse
 		View->RoughnessOverrideParameter = FVector2D( 1.0f, 0.0f );
 
+		if (bCaptureEmissiveOnly)
+		{
+			View->DiffuseOverrideParameter = FVector4(0, 0, 0, 0);
+			View->SpecularOverrideParameter = FVector4(0, 0, 0, 0);
+		}
+
 		View->bIsReflectionCapture = true;
 		View->StartFinalPostprocessSettings(CapturePosition);
 		View->EndFinalPostprocessSettings();
@@ -1465,7 +1471,7 @@ void FScene::UpdateReflectionCaptureContents(UReflectionCaptureComponent* Captur
 				ClearScratchCubemaps();
 			});
 			
-			CaptureSceneIntoScratchCubemap(this, CaptureComponent->GetComponentLocation(), false, 0, false);
+			CaptureSceneIntoScratchCubemap(this, CaptureComponent->GetComponentLocation(), false, 0, false, false);
 
 			ENQUEUE_UNIQUE_RENDER_COMMAND( 
 				FilterCommand,
@@ -1499,9 +1505,9 @@ void FScene::UpdateReflectionCaptureContents(UReflectionCaptureComponent* Captur
 	}
 }
 
-void CopyToSkyTexture(FScene* Scene, FSkyLightSceneProxy* SkyProxy)
+void CopyToSkyTexture(FScene* Scene, FTexture* ProcessedTexture)
 {
-	const int32 EffectiveTopMipSize = SkyProxy->ProcessedTexture->GetSizeX();
+	const int32 EffectiveTopMipSize = ProcessedTexture->GetSizeX();
 	const int32 NumMips = FMath::CeilLogTwo(EffectiveTopMipSize) + 1;
 
 	// GPU copy back to the skylight's texture, which is not a render target
@@ -1512,12 +1518,12 @@ void CopyToSkyTexture(FScene* Scene, FSkyLightSceneProxy* SkyProxy)
 
 		for (int32 CubeFace = 0; CubeFace < CubeFace_MAX; CubeFace++)
 		{
-			RHICopyToResolveTarget(EffectiveSource.ShaderResourceTexture, SkyProxy->ProcessedTexture->TextureRHI, true, FResolveParams(FResolveRect(), (ECubeFace)CubeFace, MipIndex, 0, 0));
+			RHICopyToResolveTarget(EffectiveSource.ShaderResourceTexture, ProcessedTexture->TextureRHI, true, FResolveParams(FResolveRect(), (ECubeFace)CubeFace, MipIndex, 0, 0));
 		}
 	}
 }
 
-void FScene::UpdateSkyCaptureContents(USkyLightComponent* CaptureComponent)
+void FScene::UpdateSkyCaptureContents(const USkyLightComponent* CaptureComponent, bool bCaptureEmissiveOnly, FTexture* OutProcessedTexture, FSHVectorRGB3& OutIrradianceEnvironmentMap)
 {
 	if (IsReflectionEnvironmentAvailable())
 	{
@@ -1530,7 +1536,7 @@ void FScene::UpdateSkyCaptureContents(USkyLightComponent* CaptureComponent)
 
 		if (CaptureComponent->SourceType == SLS_CapturedScene)
 		{
-			CaptureSceneIntoScratchCubemap(this, CaptureComponent->GetComponentLocation(), true, CaptureComponent->SkyDistanceThreshold, CaptureComponent->bLowerHemisphereIsBlack);
+			CaptureSceneIntoScratchCubemap(this, CaptureComponent->GetComponentLocation(), true, CaptureComponent->SkyDistanceThreshold, CaptureComponent->bLowerHemisphereIsBlack, bCaptureEmissiveOnly);
 		}
 		else if (CaptureComponent->SourceType == SLS_SpecifiedCubemap)
 		{
@@ -1547,34 +1553,27 @@ void FScene::UpdateSkyCaptureContents(USkyLightComponent* CaptureComponent)
 			check(0);
 		}
 
-		FSHVectorRGB3 IrradianceEnvironmentMap;
-
 		ENQUEUE_UNIQUE_RENDER_COMMAND_ONEPARAMETER( 
 			FilterCommand,
-			FSHVectorRGB3*, IrradianceEnvironmentMap, &IrradianceEnvironmentMap,
+			FSHVectorRGB3*, IrradianceEnvironmentMap, &OutIrradianceEnvironmentMap,
 		{
 			FilterReflectionEnvironment(IrradianceEnvironmentMap);
 		});
 
-		// Wait until the SH coefficients have been written out by the RT
-		//@todo - remove the need for this
+		// Wait until the SH coefficients have been written out by the RT before returning
+		//@todo realtime skylight updates - remove the need for this
 		FlushRenderingCommands();
 
-		CaptureComponent->SetIrradianceEnvironmentMap(IrradianceEnvironmentMap);
-
-		FSkyLightSceneProxy* SkyProxy = new FSkyLightSceneProxy(CaptureComponent);
-
-		ENQUEUE_UNIQUE_RENDER_COMMAND_TWOPARAMETER( 
-			CopyCommand,
-			FScene*, Scene, this,
-			FSkyLightSceneProxy*, SkyProxy, SkyProxy,
+		// Optionally copy the filtered mip chain to the output texture
+		if (OutProcessedTexture)
 		{
-			CopyToSkyTexture(Scene, SkyProxy);
-
-			// Clean up the proxy now that the rendering thread is done with it
-			delete SkyProxy;
-		});
-
-		CaptureComponent->MarkRenderStateDirty();
+			ENQUEUE_UNIQUE_RENDER_COMMAND_TWOPARAMETER( 
+				CopyCommand,
+				FScene*, Scene, this,
+				FTexture*, ProcessedTexture, OutProcessedTexture,
+			{
+				CopyToSkyTexture(Scene, ProcessedTexture);
+			});
+		}
 	}
 }
diff --git a/Engine/Source/Runtime/Renderer/Private/ScenePrivate.h b/Engine/Source/Runtime/Renderer/Private/ScenePrivate.h
--- a/Engine/Source/Runtime/Renderer/Private/ScenePrivate.h
+++ b/Engine/Source/Runtime/Renderer/Private/ScenePrivate.h
@@ -1307,7 +1307,7 @@ public:
 	virtual void UpdateSceneCaptureContents(class USceneCaptureComponent2D* CaptureComponent);
 	virtual void UpdateSceneCaptureContents(class USceneCaptureComponentCube* CaptureComponent);
 	virtual void AllocateReflectionCaptures(const TArray<UReflectionCaptureComponent*>& NewCaptures);
-	virtual void UpdateSkyCaptureContents(USkyLightComponent* CaptureComponent);
+	virtual void UpdateSkyCaptureContents(const USkyLightComponent* CaptureComponent, bool bCaptureEmissiveOnly, FTexture* OutProcessedTexture, FSHVectorRGB3& OutIrradianceEnvironmentMap);
 	virtual void AddPrecomputedLightVolume(const class FPrecomputedLightVolume* Volume);
 	virtual void RemovePrecomputedLightVolume(const class FPrecomputedLightVolume* Volume);
 	virtual void UpdateLightTransform(ULightComponent* Light);
diff --git a/Engine/Source/Runtime/Renderer/Private/SceneRendering.cpp b/Engine/Source/Runtime/Renderer/Private/SceneRendering.cpp
--- a/Engine/Source/Runtime/Renderer/Private/SceneRendering.cpp
+++ b/Engine/Source/Runtime/Renderer/Private/SceneRendering.cpp
@@ -141,6 +141,8 @@ void FViewInfo::SetupSkyIrradianceEnvironmentMapConstants(FVector4* OutSkyIrradi
 
 	if (Scene 
 		&& Scene->SkyLight 
+		// Skylights with static lighting already had their diffuse contribution baked into lightmaps
+		&& !Scene->SkyLight->bHasStaticLighting
 		&& Family->EngineShowFlags.SkyLighting)
 	{
 		const FSHVectorRGB3& SkyIrradiance = Scene->SkyLight->IrradianceEnvironmentMap;
