Static sky lights get baked into character lighting samples
Stationary sky lights get baked into the low quality character lighting samples
Fixed a bug where all character lighting sample lighting was half bright

diff --git a/Engine/Source/Programs/UnrealLightmass/Private/Lighting/FinalGather.cpp b/Engine/Source/Programs/UnrealLightmass/Private/Lighting/FinalGather.cpp
--- a/Engine/Source/Programs/UnrealLightmass/Private/Lighting/FinalGather.cpp
+++ b/Engine/Source/Programs/UnrealLightmass/Private/Lighting/FinalGather.cpp
@@ -50,6 +50,7 @@ void FStaticLightingSystem::CalculateVolumeSampleIncidentRadiance(
 	RepresentativeVertex.GenerateTriangleTangents();
 
 	FGatheredLightSample UpperStaticDirectLighting;
+	// Stationary point and spot light direct contribution
 	FGatheredLightSample UpperToggleableDirectLighting;
 	float UpperToggleableDirectionalLightShadowing = 1;
 
@@ -81,6 +82,7 @@ void FStaticLightingSystem::CalculateVolumeSampleIncidentRadiance(
 	RepresentativeVertex.GenerateTriangleTangents();
 
 	FGatheredLightSample LowerStaticDirectLighting;
+	// Stationary point and spot light direct contribution
 	FGatheredLightSample LowerToggleableDirectLighting;
 	float LowerToggleableDirectionalLightShadowing = 1;
 
@@ -101,9 +103,13 @@ void FStaticLightingSystem::CalculateVolumeSampleIncidentRadiance(
 		bDebugThisSample
 		);
 
-	const FGatheredLightSample CombinedIndirectLighting = (UpperHemisphereSample + LowerHemisphereSample) * .5f;;
-	const FGatheredLightSample CombinedHighQualitySample = (UpperStaticDirectLighting + LowerStaticDirectLighting) * .5f + CombinedIndirectLighting;
-	const FGatheredLightSample CombinedLowQualitySample = (UpperStaticDirectLighting + UpperToggleableDirectLighting + LowerStaticDirectLighting + LowerToggleableDirectLighting) * .5f + CombinedIndirectLighting;
+	const FGatheredLightSample CombinedIndirectLighting = UpperHemisphereSample + LowerHemisphereSample;
+	const FGatheredLightSample CombinedHighQualitySample = UpperStaticDirectLighting + LowerStaticDirectLighting + CombinedIndirectLighting;
+	 
+	// Composite point and spot stationary direct lighting into the low quality volume samples, since we won't be applying them dynamically
+	FGatheredLightSample CombinedLowQualitySample = UpperStaticDirectLighting + UpperToggleableDirectLighting + LowerStaticDirectLighting + LowerToggleableDirectLighting + CombinedIndirectLighting;
+	// Composite stationary sky light contribution to the low quality volume samples, since we won't be applying it dynamically
+	CombinedLowQualitySample = CombinedLowQualitySample + UpperHemisphereSample.StationarySkyLighting + LowerHemisphereSample.StationarySkyLighting;
 
 	for (int32 CoefficientIndex = 0; CoefficientIndex < 4; CoefficientIndex++)
 	{
@@ -1364,6 +1370,20 @@ SampleType FStaticLightingSystem::IncomingRadianceUniform(
 			}
 		}
 
+		FLinearColor StaticSkyLighting = FLinearColor::Black;
+		FLinearColor StationarySkyLighting = FLinearColor::Black;
+		EvaluateSkyLighting(WorldPathDirection, RayIntersection.bIntersects, StaticSkyLighting, StationarySkyLighting);
+
+		if (StaticSkyLighting != FLinearColor::Black)
+		{
+			IncomingRadiance.AddIncomingRadiance(StaticSkyLighting, SampleWeight, TangentPathDirection, WorldPathDirection);
+		}
+
+		if (StationarySkyLighting != FLinearColor::Black)
+		{
+			IncomingRadiance.AddIncomingStationarySkyLight(StationarySkyLighting, SampleWeight, TangentPathDirection, WorldPathDirection);
+		}
+
 #if ALLOW_LIGHTMAP_SAMPLE_DEBUGGING
 		if (bDebugThisTexel
 			&& GeneralSettings.ViewSingleBounceNumber == BounceNumber
