#ttp 331438 UE4: PlatformMath refactor

diff --git a/Engine/Source/Editor/Matinee/Private/Matinee.cpp b/Engine/Source/Editor/Matinee/Private/Matinee.cpp
--- a/Engine/Source/Editor/Matinee/Private/Matinee.cpp
+++ b/Engine/Source/Editor/Matinee/Private/Matinee.cpp
@@ -2238,7 +2238,7 @@ void FMatinee::DrawModeHUD(FLevelEditorViewportClient* ViewportClient,FViewport*
 		// Draw subtitles (toggle is handled internally)
 		FVector2D MinPos(0.f, 0.f);
 		FVector2D MaxPos(1.f, .9f);
-		FIntRect SubtitleRegion(FMath::Trunc(SizeX * MinPos.X), FMath::Trunc(SizeY * MinPos.Y), FMath::Trunc(SizeX * MaxPos.X), FMath::Trunc(SizeY * MaxPos.Y));
+		FIntRect SubtitleRegion(FMath::TruncToInt(SizeX * MinPos.X), FMath::TruncToInt(SizeY * MinPos.Y), FMath::TruncToInt(SizeX * MaxPos.X), FMath::TruncToInt(SizeY * MaxPos.Y));
 		FSubtitleManager::GetSubtitleManager()->DisplaySubtitles( Canvas, SubtitleRegion, ViewportClient->GetWorld()->GetAudioTimeSeconds() );
 	}
 	// Camera Shot Names
diff --git a/Engine/Source/Editor/Matinee/Private/MatineeDraw.cpp b/Engine/Source/Editor/Matinee/Private/MatineeDraw.cpp
--- a/Engine/Source/Editor/Matinee/Private/MatineeDraw.cpp
+++ b/Engine/Source/Editor/Matinee/Private/MatineeDraw.cpp
@@ -88,7 +88,7 @@ int32 FMatineeViewportClient::DrawLabel( FCanvas* Canvas, float StartX, float St
 	int32 Result = 0;
 	if( LabelFont != NULL )
 	{
-		Result = Canvas->DrawShadowedString( FMath::TruncFloat(StartX), FMath::TruncFloat(StartY), Text, LabelFont, Color );
+		Result = Canvas->DrawShadowedString( FMath::TruncToFloat(StartX), FMath::TruncToFloat(StartY), Text, LabelFont, Color );
 	}
 	return Result;
 }
@@ -109,7 +109,7 @@ float FMatineeViewportClient::GetGridSpacing(int32 GridNum)
 /** Calculate the best frames' density. */
 uint32 FMatineeViewportClient::CalculateBestFrameStep(float SnapAmount, float PixelsPerSec, float MinPixelsPerGrid)
 {
-	uint32 FrameRate = FMath::Ceil(1.0f / SnapAmount);
+	uint32 FrameRate = FMath::CeilToInt(1.0f / SnapAmount);
 	uint32 FrameStep = 1;
 	
 	// Calculate minimal-symmetric integer divisor.
@@ -443,13 +443,13 @@ void FMatinee::InvalidateTrackWindowViewports()
 FString FMatinee::MakeTimecodeString( float InTime, bool bIncludeMinutes ) const
 {
 	// SMPTE-style timecode
-	const int32 MinutesVal = FMath::Trunc( InTime / 60.0f );
-	const int32 SecondsVal = FMath::Trunc( InTime );
+	const int32 MinutesVal = FMath::TruncToInt( InTime / 60.0f );
+	const int32 SecondsVal = FMath::TruncToInt( InTime );
 	const float Frames = InTime / SnapAmount;
-	const int32 FramesVal = FMath::Round( Frames );
+	const int32 FramesVal = FMath::RoundToInt( Frames );
 	const float Subseconds = FMath::Fractional( InTime );
 	const float SubsecondFrames = Subseconds / SnapAmount;
-	const int32 SubsecondFramesVal = FMath::Round( SubsecondFrames );
+	const int32 SubsecondFramesVal = FMath::RoundToInt( SubsecondFrames );
 	const float SubframeDiff = Frames - ( float )FramesVal;
 
 	// Are we currently between frames?
@@ -516,7 +516,7 @@ void FMatineeViewportClient::DrawGrid(FViewport* Viewport, FCanvas* Canvas, bool
 	FCanvasLineItem LineItem;									
 	FCanvasTextItem TextItem( FVector2D::ZeroVector, FText::GetEmpty(), GEditor->GetSmallFont(), FLinearColor::Gray );
 	
-	int32 LineNum = FMath::Floor(InterpEd->ViewStartTime/GridSpacing);
+	int32 LineNum = FMath::FloorToInt(InterpEd->ViewStartTime/GridSpacing);
 	while( LineNum*GridSpacing < InterpEd->ViewEndTime )
 	{
 		float LineTime = LineNum*GridSpacing;
@@ -795,7 +795,7 @@ void FMatineeViewportClient::DrawTimeline(FViewport* Viewport, FCanvas* Canvas)
 				FString StartString = FString::Printf( TEXT("%3.2fs"), KeyStartTime );
 				if ( InterpEd->bSnapToFrames )
 				{
-					StartString += FString::Printf( TEXT(" / %df"), FMath::Round(KeyStartTime / InterpEd->SnapAmount));
+					StartString += FString::Printf( TEXT(" / %df"), FMath::RoundToInt(KeyStartTime / InterpEd->SnapAmount));
 				}
 				StringSize( LabelFont, XL, YL, *StartString);
 				DrawLabel(Canvas, KeyStartX - XL, ViewY - TotalBarHeight - RangeTickHeight - YL - 2, *StartString, KeyRangeMarkerColor);
@@ -809,7 +809,7 @@ void FMatineeViewportClient::DrawTimeline(FViewport* Viewport, FCanvas* Canvas)
 				FString EndString = FString::Printf( TEXT("%3.2fs"), KeyEndTime );
 				if ( InterpEd->bSnapToFrames )
 				{
-					EndString += FString::Printf( TEXT(" / %df"), FMath::Round(KeyEndTime / InterpEd->SnapAmount));
+					EndString += FString::Printf( TEXT(" / %df"), FMath::RoundToInt(KeyEndTime / InterpEd->SnapAmount));
 				}
 
 				StringSize( LabelFont, XL, YL, *EndString);
@@ -825,7 +825,7 @@ void FMatineeViewportClient::DrawTimeline(FViewport* Viewport, FCanvas* Canvas)
 			FString RangeString = FString::Printf( TEXT("%3.2fs"), KeyRange );
 			if ( InterpEd->bSnapToFrames )
 			{
-				RangeString += FString::Printf( TEXT(" / %df"), FMath::Round(KeyRange / InterpEd->SnapAmount));
+				RangeString += FString::Printf( TEXT(" / %df"), FMath::RoundToInt(KeyRange / InterpEd->SnapAmount));
 			}
 
 			StringSize( LabelFont, XL, YL, *RangeString);
@@ -850,7 +850,7 @@ void FMatineeViewportClient::DrawTimeline(FViewport* Viewport, FCanvas* Canvas)
 				FString KeyString = FString::Printf( TEXT("%3.2fs"), KeyTime );				
 				if ( InterpEd->bSnapToFrames )
 				{
-					KeyString += FString::Printf( TEXT(" / %df"), FMath::Round(KeyTime / InterpEd->SnapAmount));
+					KeyString += FString::Printf( TEXT(" / %df"), FMath::RoundToInt(KeyTime / InterpEd->SnapAmount));
 				}
 
 				StringSize( LabelFont, XL, YL, *KeyString);
@@ -1147,7 +1147,7 @@ void FMatineeViewportClient::DrawSubTrackGroup( FCanvas* Canvas, UInterpTrack* T
 		TextColor = FLinearColor(0.5f,0.5f,0.5f);
 	}
 
-	DrawLabel(Canvas, -InterpEd->LabelWidth + TrackTitleIndentPixels, FMath::TruncFloat(0.5*TrackHeight - 0.5*YL), *Text, TextColor);
+	DrawLabel(Canvas, -InterpEd->LabelWidth + TrackTitleIndentPixels, FMath::TruncToFloat(0.5*TrackHeight - 0.5*YL), *Text, TextColor);
 	if( Canvas->IsHitTesting() )
 	{
 		Canvas->SetHitProxy( NULL );
@@ -2237,7 +2237,7 @@ void FMatineeViewportClient::DrawSubTrackGroupKeys( FCanvas* Canvas, UInterpTrac
 				// If the group is collapsed draw each keyframe as a triangle
 				const int32 KeyHalfTriSize = 6;
 				const int32 KeyVertOffset = 3;
-				int32 PixelPos = FMath::Trunc(DrawInfo.KeyTime * InterpEd->PixelsPerSec);
+				int32 PixelPos = FMath::TruncToInt(DrawInfo.KeyTime * InterpEd->PixelsPerSec);
 
 				FIntPoint A(PixelPos - KeyHalfTriSize,	TrackHeight - KeyVertOffset);
 				FIntPoint B(PixelPos + KeyHalfTriSize,	TrackHeight - KeyVertOffset);
diff --git a/Engine/Source/Editor/Matinee/Private/MatineeMenus.cpp b/Engine/Source/Editor/Matinee/Private/MatineeMenus.cpp
--- a/Engine/Source/Editor/Matinee/Private/MatineeMenus.cpp
+++ b/Engine/Source/Editor/Matinee/Private/MatineeMenus.cpp
@@ -3903,7 +3903,7 @@ void FMatinee::OnExportSoundCueInfoCommand()
 
 
 										// Also store values as frame numbers instead of time values if a frame rate is selected
-										const int32 SoundFrameIndex = bSnapToFrames ? FMath::Trunc( CurSound.Time / SnapAmount ) : 0;
+										const int32 SoundFrameIndex = bSnapToFrames ? FMath::TruncToInt( CurSound.Time / SnapAmount ) : 0;
 
 										FString TextLine = FString::Printf(
 											TEXT( "%s,%s,%s,%0.2f,%i" ),
@@ -3917,7 +3917,7 @@ void FMatinee::OnExportSoundCueInfoCommand()
 										if( FoundAnimName.Len() > 0 )
 										{
 											// Also store values as frame numbers instead of time values if a frame rate is selected
-											const int32 AnimFrameIndex = bSnapToFrames ? FMath::Trunc( FoundAnimTime / SnapAmount ) : 0;
+											const int32 AnimFrameIndex = bSnapToFrames ? FMath::TruncToInt( FoundAnimTime / SnapAmount ) : 0;
 
 											TextLine += FString::Printf(
 												TEXT( ",%s,%.2f,%i" ),
@@ -4513,7 +4513,7 @@ void FMatinee::NormalizeVelocity()
 				ZAxisTrack->Modify();
 			
 				const float TotalTime = FullEndTime - FullStartTime;
-				const int32 NumSteps = FMath::Ceil(TotalTime/(1.0f/60.0f));
+				const int32 NumSteps = FMath::CeilToInt(TotalTime/(1.0f/60.0f));
 				const float Interval = (SegmentEndTime-SegmentStartTime)/NumSteps; 
 
 				// An array of points that were created in order to normalize velocity
diff --git a/Engine/Source/Editor/Matinee/Private/MatineeTools.cpp b/Engine/Source/Editor/Matinee/Private/MatineeTools.cpp
--- a/Engine/Source/Editor/Matinee/Private/MatineeTools.cpp
+++ b/Engine/Source/Editor/Matinee/Private/MatineeTools.cpp
@@ -1645,7 +1645,7 @@ void FMatinee::SetShowTimeCursorPosForAllKeys( bool bInValue )
 float FMatinee::SnapTimeToNearestFrame( float InTime ) const
 {
 	// Compute the new time value by rounding
-	const int32 InterpPositionInFrames = FMath::Round( InTime / SnapAmount );
+	const int32 InterpPositionInFrames = FMath::RoundToInt( InTime / SnapAmount );
 	const float NewTime = InterpPositionInFrames * SnapAmount;
 
 	return NewTime;
diff --git a/Engine/Source/Editor/Matinee/Private/MatineeViewportClient.cpp b/Engine/Source/Editor/Matinee/Private/MatineeViewportClient.cpp
--- a/Engine/Source/Editor/Matinee/Private/MatineeViewportClient.cpp
+++ b/Engine/Source/Editor/Matinee/Private/MatineeViewportClient.cpp
@@ -1026,7 +1026,7 @@ void FMatineeViewportClient::Tick(float DeltaSeconds)
 		}
 
 		// Apply integer part of ScrollAccum to the curve editor view position.
-		const int32 DeltaX = FMath::Floor(ScrollAccum.X);
+		const int32 DeltaX = FMath::FloorToInt(ScrollAccum.X);
 		ScrollAccum.X -= DeltaX;
 
 		if(bNavigating)
diff --git a/Engine/Source/Editor/MeshPaint/Private/MeshPaintEdMode.cpp b/Engine/Source/Editor/MeshPaint/Private/MeshPaintEdMode.cpp
--- a/Engine/Source/Editor/MeshPaint/Private/MeshPaintEdMode.cpp
+++ b/Engine/Source/Editor/MeshPaint/Private/MeshPaintEdMode.cpp
@@ -909,10 +909,10 @@ bool FEdModeMeshPaint::PaintVertex( const FVector& InVertexPosition,
 
 
 			// Save the new color
-			InOutVertexColor.R = FMath::Clamp( FMath::Round( NewColor.R * 255.0f ), 0, 255 );
-			InOutVertexColor.G = FMath::Clamp( FMath::Round( NewColor.G * 255.0f ), 0, 255 );
-			InOutVertexColor.B = FMath::Clamp( FMath::Round( NewColor.B * 255.0f ), 0, 255 );
-			InOutVertexColor.A = FMath::Clamp( FMath::Round( NewColor.A * 255.0f ), 0, 255 );
+			InOutVertexColor.R = FMath::Clamp( FMath::RoundToInt( NewColor.R * 255.0f ), 0, 255 );
+			InOutVertexColor.G = FMath::Clamp( FMath::RoundToInt( NewColor.G * 255.0f ), 0, 255 );
+			InOutVertexColor.B = FMath::Clamp( FMath::RoundToInt( NewColor.B * 255.0f ), 0, 255 );
+			InOutVertexColor.A = FMath::Clamp( FMath::RoundToInt( NewColor.A * 255.0f ), 0, 255 );
 
 
 			// UE_LOG(LogMeshPaintEdMode, Log,  TEXT( "Painted Vertex:  OldColor=[%.2f,%.2f,%.2f,%.2f], NewColor=[%.2f,%.2f,%.2f,%.2f]" ), OldColor.R, OldColor.G, OldColor.B, OldColor.A, NewColor.R, NewColor.G, NewColor.B, NewColor.A );
@@ -2215,20 +2215,20 @@ void FEdModeMeshPaint::PaintTexture( const FMeshPaintParameters& InParams,
 		FVector2D UVOffset( 0.0f, 0.0f );
 		if( UVMax.X > 1.0f )
 		{
-			UVOffset.X = -FMath::Floor( UVMin.X );
+			UVOffset.X = -FMath::FloorToFloat( UVMin.X );
 		}
 		else if( UVMin.X < 0.0f )
 		{
-			UVOffset.X = 1.0f + FMath::Floor( -UVMax.X );
+			UVOffset.X = 1.0f + FMath::FloorToFloat( -UVMax.X );
 		}
 			
 		if( UVMax.Y > 1.0f )
 		{
-			UVOffset.Y = -FMath::Floor( UVMin.Y );
+			UVOffset.Y = -FMath::FloorToFloat( UVMin.Y );
 		}
 		else if( UVMin.Y < 0.0f )
 		{
-			UVOffset.Y = 1.0f + FMath::Floor( -UVMax.Y );
+			UVOffset.Y = 1.0f + FMath::FloorToFloat( -UVMax.Y );
 		}
 		
 		// Note that we "wrap" the texture coordinates here to handle the case where the user
@@ -3474,20 +3474,20 @@ bool FEdModeMeshPaint::GenerateSeamMask(UStaticMeshComponent* StaticMeshComponen
 			FVector2D UVOffset( 0.0f, 0.0f );
 			if( UVMax.X > 1.0f )
 			{
-				UVOffset.X = -FMath::Floor( UVMin.X );
+				UVOffset.X = -FMath::FloorToInt( UVMin.X );
 			}
 			else if( UVMin.X < 0.0f )
 			{
-				UVOffset.X = 1.0f + FMath::Floor( -UVMax.X );
+				UVOffset.X = 1.0f + FMath::FloorToInt( -UVMax.X );
 			}
 
 			if( UVMax.Y > 1.0f )
 			{
-				UVOffset.Y = -FMath::Floor( UVMin.Y );
+				UVOffset.Y = -FMath::FloorToInt( UVMin.Y );
 			}
 			else if( UVMin.Y < 0.0f )
 			{
-				UVOffset.Y = 1.0f + FMath::Floor( -UVMax.Y );
+				UVOffset.Y = 1.0f + FMath::FloorToInt( -UVMax.Y );
 			}
 
 
diff --git a/Engine/Source/Editor/NewsFeed/Private/Implementation/CdnNewsFeedTitleFile.cpp b/Engine/Source/Editor/NewsFeed/Private/Implementation/CdnNewsFeedTitleFile.cpp
--- a/Engine/Source/Editor/NewsFeed/Private/Implementation/CdnNewsFeedTitleFile.cpp
+++ b/Engine/Source/Editor/NewsFeed/Private/Implementation/CdnNewsFeedTitleFile.cpp
@@ -307,7 +307,7 @@ void FCdnNewsFeedTitleFile::EnumerateFiles_HttpRequestComplete(FHttpRequestPtr H
 						}
 						if (JsonFileHeader->HasField(TEXT("length")))
 						{
-							FileHeader.FileSize = FMath::Trunc(JsonFileHeader->GetNumberField(TEXT("length")));
+							FileHeader.FileSize = FMath::TruncToInt(JsonFileHeader->GetNumberField(TEXT("length")));
 						}
 
 						if (FileHeader.FileName.IsEmpty())
diff --git a/Engine/Source/Editor/Persona/Private/AnimationEditorViewportClient.cpp b/Engine/Source/Editor/Persona/Private/AnimationEditorViewportClient.cpp
--- a/Engine/Source/Editor/Persona/Private/AnimationEditorViewportClient.cpp
+++ b/Engine/Source/Editor/Persona/Private/AnimationEditorViewportClient.cpp
@@ -845,9 +845,9 @@ void FAnimationViewportClient::DisplayInfo(FCanvas* Canvas, FSceneView* View, bo
 
 		CurYOffset += YL + 2;
 		InfoString = FString::Printf(TEXT("Approximate Size: %ix%ix%i"), 
-			FMath::Round(PreviewSkelMeshComp->Bounds.BoxExtent.X * 2.0f),
-			FMath::Round(PreviewSkelMeshComp->Bounds.BoxExtent.Y * 2.0f),
-			FMath::Round(PreviewSkelMeshComp->Bounds.BoxExtent.Z * 2.0f));
+			FMath::RoundToInt(PreviewSkelMeshComp->Bounds.BoxExtent.X * 2.0f),
+			FMath::RoundToInt(PreviewSkelMeshComp->Bounds.BoxExtent.Y * 2.0f),
+			FMath::RoundToInt(PreviewSkelMeshComp->Bounds.BoxExtent.Z * 2.0f));
 		Canvas->DrawShadowedString( CurXOffset, CurYOffset, *InfoString, GEngine->GetSmallFont(), TextColor );
 
 		uint32 NumNotiesWithErrors = PreviewSkelMeshComp->AnimNotifyErrors.Num();
diff --git a/Engine/Source/Editor/Persona/Private/SAnimationBlendSpace.cpp b/Engine/Source/Editor/Persona/Private/SAnimationBlendSpace.cpp
--- a/Engine/Source/Editor/Persona/Private/SAnimationBlendSpace.cpp
+++ b/Engine/Source/Editor/Persona/Private/SAnimationBlendSpace.cpp
@@ -1458,8 +1458,8 @@ bool SBlendSpaceGridWidget::GetMinMaxFromGridPos(FVector GridPos, FVector2D & Wi
 	FVector NormalizedGridPos = NormalizeGridPos(GridDim, GridPos);
 	FVector Multiplier = NormalizedGridPos/InvGridSize;
 	FVector GridMin, GridMax;
-	GridMin.X = InvGridSize.X * FGenericPlatformMath::Trunc(Multiplier.X);
-	GridMin.Y = InvGridSize.Y * FGenericPlatformMath::Trunc(Multiplier.Y);
+	GridMin.X = InvGridSize.X * FGenericPlatformMath::TruncToFloat(Multiplier.X);
+	GridMin.Y = InvGridSize.Y * FGenericPlatformMath::TruncToFloat(Multiplier.Y);
 	GridMin.Z = 0.f;
 
 	GridMax = GridMin + InvGridSize;
diff --git a/Engine/Source/Editor/Persona/Private/SAnimationBlendSpace1D.cpp b/Engine/Source/Editor/Persona/Private/SAnimationBlendSpace1D.cpp
--- a/Engine/Source/Editor/Persona/Private/SAnimationBlendSpace1D.cpp
+++ b/Engine/Source/Editor/Persona/Private/SAnimationBlendSpace1D.cpp
@@ -480,7 +480,7 @@ FVector SBlendSpace1DWidget::SnapEditorPosToGrid(const FVector& InPos) const
 
 	const float EditorStep = Param.GetRange() / (float)Param.GridNum;
 	const float DiffFromStart = InPos.X - Param.Min;
-	const float SnappedToStep = FMath::Round(DiffFromStart / EditorStep) * EditorStep;
+	const float SnappedToStep = FMath::RoundToFloat(DiffFromStart / EditorStep) * EditorStep;
 
 	SnappedVector.X = SnappedToStep + Param.Min;
 	return SnappedVector;
diff --git a/Engine/Source/Editor/Sequencer/Private/SSection.cpp b/Engine/Source/Editor/Sequencer/Private/SSection.cpp
--- a/Engine/Source/Editor/Sequencer/Private/SSection.cpp
+++ b/Engine/Source/Editor/Sequencer/Private/SSection.cpp
@@ -106,7 +106,7 @@ FSelectedKey SSection::GetKeyUnderMouse( const FVector2D& MousePosition, const F
 				float KeyPosition = TimeToPixelConverter.TimeToPixel( KeyArea->GetKeyTime(KeyHandle) );
 
 				FGeometry KeyGeometry = KeyAreaGeometry.MakeChild( 
-					FVector2D( KeyPosition - FMath::Trunc(SequencerSectionConstants::KeySize.X/2.0f), ((KeyAreaGeometry.Size.Y*.5f)-(SequencerSectionConstants::KeySize.Y*.5f)) ),
+					FVector2D( KeyPosition - FMath::TruncToFloat(SequencerSectionConstants::KeySize.X/2.0f), ((KeyAreaGeometry.Size.Y*.5f)-(SequencerSectionConstants::KeySize.Y*.5f)) ),
 					SequencerSectionConstants::KeySize );
 
 				if( KeyGeometry.IsUnderLocation( MousePosition ) )
@@ -378,7 +378,7 @@ void SSection::PaintKeys( const FGeometry& AllottedGeometry, const FSlateRect& M
 					// always draw selected keys on top of other keys
 					bSelected ? KeyLayer+1 : KeyLayer,
 					// Center the key along Y.  Ensure the middle of the key is at the actual key time
-					KeyAreaGeometry.ToPaintGeometry( FVector2D( KeyPosition - FMath::Trunc(SequencerSectionConstants::KeySize.X/2.0f), ((KeyAreaGeometry.Size.Y*.5f)-(SequencerSectionConstants::KeySize.Y*.5f)) ), SequencerSectionConstants::KeySize ),
+					KeyAreaGeometry.ToPaintGeometry( FVector2D( KeyPosition - FMath::TruncToFloat(SequencerSectionConstants::KeySize.X/2.0f), ((KeyAreaGeometry.Size.Y*.5f)-(SequencerSectionConstants::KeySize.Y*.5f)) ), SequencerSectionConstants::KeySize ),
 					KeyBrush,
 					MyClippingRect,
 					ESlateDrawEffect::None,
diff --git a/Engine/Source/Editor/Sequencer/Private/SectionDragOperations.cpp b/Engine/Source/Editor/Sequencer/Private/SectionDragOperations.cpp
--- a/Engine/Source/Editor/Sequencer/Private/SectionDragOperations.cpp
+++ b/Engine/Source/Editor/Sequencer/Private/SectionDragOperations.cpp
@@ -284,7 +284,7 @@ void FMoveSection::OnDrag( const FPointerEvent& MouseEvent, const FVector2D& Loc
 					}
 				}
 
-				TargetRowIndex = FMath::Clamp(Section->GetRowIndex() + FMath::Floor(LocalMousePos.Y / TrackHeight),
+				TargetRowIndex = FMath::Clamp(Section->GetRowIndex() + FMath::FloorToInt(LocalMousePos.Y / TrackHeight),
 					0, MaxRowIndex + 1);
 			}
 		}
diff --git a/Engine/Source/Editor/Sequencer/Private/SequencerDisplayNode.cpp b/Engine/Source/Editor/Sequencer/Private/SequencerDisplayNode.cpp
--- a/Engine/Source/Editor/Sequencer/Private/SequencerDisplayNode.cpp
+++ b/Engine/Source/Editor/Sequencer/Private/SequencerDisplayNode.cpp
@@ -75,7 +75,7 @@ int32 SSequencerObjectTrack::OnPaint(const FGeometry& AllottedGeometry, const FS
 		FSlateDrawElement::MakeBox(
 			OutDrawElements,
 			LayerId+1,
-			AllottedGeometry.ToPaintGeometry(FVector2D(KeyPosition - FMath::Trunc(KeyMarkSize.X/2.f), FMath::Trunc(AllottedGeometry.Size.Y/2.f - KeyMarkSize.Y/2.f)), KeyMarkSize),
+			AllottedGeometry.ToPaintGeometry(FVector2D(KeyPosition - FMath::TruncToFloat(KeyMarkSize.X/2.f), FMath::TruncToFloat(AllottedGeometry.Size.Y/2.f - KeyMarkSize.Y/2.f)), KeyMarkSize),
 			FEditorStyle::GetBrush("Sequencer.KeyMark"),
 			MyClippingRect,
 			ESlateDrawEffect::None,
diff --git a/Engine/Source/Editor/Sequencer/Private/TimeSliderController.cpp b/Engine/Source/Editor/Sequencer/Private/TimeSliderController.cpp
--- a/Engine/Source/Editor/Sequencer/Private/TimeSliderController.cpp
+++ b/Engine/Source/Editor/Sequencer/Private/TimeSliderController.cpp
@@ -136,7 +136,7 @@ void FSequencerTimeSliderController::DrawTicks( FSlateWindowElementList& OutDraw
 	// For slightly larger halfway tick mark
 	const int32 HalfDivider = Divider / 2;
 	// Find out where to start from
-	int32 OffsetNum = FMath::Floor(RangeToScreen.ViewInput.GetLowerBoundValue() / Spacing);
+	int32 OffsetNum = FMath::FloorToInt(RangeToScreen.ViewInput.GetLowerBoundValue() / Spacing);
 	
 	FSlateFontInfo SmallLayoutFont( FPaths::EngineContentDir() / TEXT("Slate/Fonts/Roboto-Regular.ttf"), 8 );
 
@@ -257,7 +257,7 @@ int32 FSequencerTimeSliderController::OnPaintTimeSlider( bool bMirrorLabels, con
 		DrawTicks( OutDrawElements, RangeToScreen, Args );
 
 		const float HandleSize = 13.0f;
-		float HalfSize = FMath::Trunc(HandleSize/2.0f);
+		float HalfSize = FMath::TruncToFloat(HandleSize/2.0f);
 
 		// Draw the scrub handle
 		const float XPos = RangeToScreen.InputToLocalX( TimeSliderArgs.ScrubPosition.Get() );
@@ -521,7 +521,7 @@ int32 FSequencerTimeSliderController::OnPaintSectionView( const FGeometry& Allot
 		TArray<FVector2D> LinePoints;
 		LinePoints.AddUninitialized(2);
 		LinePoints[0] = FVector2D( 1.0f, 0.0f );
-		LinePoints[1] = FVector2D( 1.0f, FMath::Round( AllottedGeometry.Size.Y ) );
+		LinePoints[1] = FVector2D( 1.0f, FMath::RoundToFloat( AllottedGeometry.Size.Y ) );
 
 		FSlateDrawElement::MakeLines(
 			OutDrawElements,
diff --git a/Engine/Source/Editor/StaticMeshEditor/Private/StaticMeshEditorViewportClient.cpp b/Engine/Source/Editor/StaticMeshEditor/Private/StaticMeshEditorViewportClient.cpp
--- a/Engine/Source/Editor/StaticMeshEditor/Private/StaticMeshEditorViewportClient.cpp
+++ b/Engine/Source/Editor/StaticMeshEditor/Private/StaticMeshEditorViewportClient.cpp
@@ -441,7 +441,7 @@ void FStaticMeshEditorViewportClient::DrawCanvas( FViewport& InViewport, FSceneV
 						//Figure out the text height
 						FTextSizingParameters Parameters(GEngine->GetSmallFont(), 1.0f, 1.0f);
 						UCanvas::CanvasStringSize(Parameters, *Socket->SocketName.ToString());
-						int32 YL = FMath::Trunc(Parameters.DrawYL);
+						int32 YL = FMath::TruncToInt(Parameters.DrawYL);
 
 						DrawAngles(&Canvas, XPos, YPos + YL, 
 							Widget->GetCurrentAxis(), 
diff --git a/Engine/Source/Editor/StatsViewer/Private/StatsPages/PrimitiveStatsPage.cpp b/Engine/Source/Editor/StatsViewer/Private/StatsPages/PrimitiveStatsPage.cpp
--- a/Engine/Source/Editor/StatsViewer/Private/StatsPages/PrimitiveStatsPage.cpp
+++ b/Engine/Source/Editor/StatsViewer/Private/StatsPages/PrimitiveStatsPage.cpp
@@ -428,7 +428,7 @@ void FPrimitiveStatsPage::GenerateTotals( const TArray< TWeakObjectPtr<UObject>
 			TotalEntry->LightsLM += StatsEntry->LightsLM;
 			TotalEntry->LightsOther += StatsEntry->LightsOther;
 			TotalEntry->LightsTotal += StatsEntry->LightsTotal;
-			TotalEntry->ObjLightCost += FMath::Trunc( StatsEntry->ObjLightCost );
+			TotalEntry->ObjLightCost += FMath::TruncToFloat( StatsEntry->ObjLightCost );
 			TotalEntry->LightMapData += StatsEntry->LightMapData;
 			TotalEntry->LMSMResolution += StatsEntry->LMSMResolution;
 			TotalEntry->RadiusMin = FMath::Min(TotalEntry->RadiusMin, StatsEntry->RadiusMin);
diff --git a/Engine/Source/Editor/UMGEditor/Private/SDesignSurface.cpp b/Engine/Source/Editor/UMGEditor/Private/SDesignSurface.cpp
--- a/Engine/Source/Editor/UMGEditor/Private/SDesignSurface.cpp
+++ b/Engine/Source/Editor/UMGEditor/Private/SDesignSurface.cpp
@@ -142,7 +142,7 @@ FReply SDesignSurface::OnMouseWheel(const FGeometry& MyGeometry, const FPointerE
 {
 	// We want to zoom into this point; i.e. keep it the same fraction offset into the panel
 	const FVector2D WidgetSpaceCursorPos = MyGeometry.AbsoluteToLocal(MouseEvent.GetScreenSpacePosition());
-	const int32 ZoomLevelDelta = FMath::Floor(MouseEvent.GetWheelDelta());
+	const int32 ZoomLevelDelta = FMath::FloorToInt(MouseEvent.GetWheelDelta());
 	ChangeZoomLevel(ZoomLevelDelta, WidgetSpaceCursorPos, MouseEvent.IsControlDown());
 
 	return FReply::Handled();
diff --git a/Engine/Source/Editor/UnrealEd/Private/DragTool_FrustumSelect.cpp b/Engine/Source/Editor/UnrealEd/Private/DragTool_FrustumSelect.cpp
--- a/Engine/Source/Editor/UnrealEd/Private/DragTool_FrustumSelect.cpp
+++ b/Engine/Source/Editor/UnrealEd/Private/DragTool_FrustumSelect.cpp
@@ -106,7 +106,7 @@ void FDragTool_ActorFrustumSelect::EndDrag()
 		}
 
 		// Extend the endpoint of the rect to get the actual line
-		FIntRect BoxRect( FIntPoint( FMath::Max( 0.0f, Start.X ), FMath::Max( 0.0f, Start.Y ) ), FIntPoint( FMath::Min(ViewportSizeX, FMath::Trunc(End.X+1)), FMath::Min( ViewportSizeY, FMath::Trunc(End.Y+1) ) ) );
+		FIntRect BoxRect( FIntPoint( FMath::Max( 0.0f, Start.X ), FMath::Max( 0.0f, Start.Y ) ), FIntPoint( FMath::Min(ViewportSizeX, FMath::TruncToInt(End.X+1)), FMath::Min( ViewportSizeY, FMath::TruncToInt(End.Y+1) ) ) );
 
 		const TArray<FColor>& RawHitProxyData = ViewportClient->Viewport->GetRawHitProxyData(BoxRect);
 
diff --git a/Engine/Source/Editor/UnrealEd/Private/DragTool_Measure.cpp b/Engine/Source/Editor/UnrealEd/Private/DragTool_Measure.cpp
--- a/Engine/Source/Editor/UnrealEd/Private/DragTool_Measure.cpp
+++ b/Engine/Source/Editor/UnrealEd/Private/DragTool_Measure.cpp
@@ -58,7 +58,7 @@ void FDragTool_Measure::Render(const FSceneView* View, FCanvas* Canvas)
 		FCanvasLineItem LineItem( PixelStart, PixelEnd );
 		Canvas->DrawItem( LineItem );
 
-		const int32 Length = FMath::Ceil((End - Start).Size());
+		const int32 Length = FMath::CeilToInt((End - Start).Size());
 		if( Length == 0 )
 		{
 			return;
@@ -69,7 +69,7 @@ void FDragTool_Measure::Render(const FSceneView* View, FCanvas* Canvas)
 		if (View->WorldToPixel(WorldMid, PixelMid))
 		{
 			FString LengthStr = FString::Printf(TEXT("%d"), Length);
-			FCanvasTextItem TextItem( FVector2D( FMath::Floor(PixelMid.X), FMath::Floor(PixelMid.Y) ), FText::FromString( LengthStr ), GEngine->GetSmallFont(), FLinearColor::White );
+			FCanvasTextItem TextItem( FVector2D( FMath::FloorToFloat(PixelMid.X), FMath::FloorToFloat(PixelMid.Y) ), FText::FromString( LengthStr ), GEngine->GetSmallFont(), FLinearColor::White );
 			TextItem.bCentreX = true;
 			Canvas->DrawItem( TextItem );
 		}
diff --git a/Engine/Source/Editor/UnrealEd/Private/EditorActor.cpp b/Engine/Source/Editor/UnrealEd/Private/EditorActor.cpp
--- a/Engine/Source/Editor/UnrealEd/Private/EditorActor.cpp
+++ b/Engine/Source/Editor/UnrealEd/Private/EditorActor.cpp
@@ -2055,9 +2055,9 @@ void UUnrealEdEngine::edactAlignOrigin()
 
 			//Snap the location of the brush to the grid
 			FVector BrushLocation = Brush->GetActorLocation();
-			BrushLocation.X = FMath::Round( BrushLocation.X  / GetGridSize() ) * GetGridSize();
-			BrushLocation.Y = FMath::Round( BrushLocation.Y  / GetGridSize() ) * GetGridSize();
-			BrushLocation.Z = FMath::Round( BrushLocation.Z  / GetGridSize() ) * GetGridSize();
+			BrushLocation.X = FMath::RoundToFloat( BrushLocation.X  / GetGridSize() ) * GetGridSize();
+			BrushLocation.Y = FMath::RoundToFloat( BrushLocation.Y  / GetGridSize() ) * GetGridSize();
+			BrushLocation.Z = FMath::RoundToFloat( BrushLocation.Z  / GetGridSize() ) * GetGridSize();
 			Brush->SetActorLocation(BrushLocation, false);
 
 			//Update EditorMode locations to match the new brush location
@@ -2101,9 +2101,9 @@ void UUnrealEdEngine::edactAlignVertices()
 				for( int32 VertIdx=0; VertIdx<Poly->Vertices.Num(); VertIdx++ )
 				{
 					// Snap each vertex to the nearest grid.
-					Poly->Vertices[VertIdx].X = FMath::Round( ( Poly->Vertices[VertIdx].X + BrushLocation.X )  / GetGridSize() ) * GetGridSize() - BrushLocation.X;
-					Poly->Vertices[VertIdx].Y = FMath::Round( ( Poly->Vertices[VertIdx].Y + BrushLocation.Y )  / GetGridSize() ) * GetGridSize() - BrushLocation.Y;
-					Poly->Vertices[VertIdx].Z = FMath::Round( ( Poly->Vertices[VertIdx].Z + BrushLocation.Z )  / GetGridSize() ) * GetGridSize() - BrushLocation.Z;
+					Poly->Vertices[VertIdx].X = FMath::RoundToFloat( ( Poly->Vertices[VertIdx].X + BrushLocation.X )  / GetGridSize() ) * GetGridSize() - BrushLocation.X;
+					Poly->Vertices[VertIdx].Y = FMath::RoundToFloat( ( Poly->Vertices[VertIdx].Y + BrushLocation.Y )  / GetGridSize() ) * GetGridSize() - BrushLocation.Y;
+					Poly->Vertices[VertIdx].Z = FMath::RoundToFloat( ( Poly->Vertices[VertIdx].Z + BrushLocation.Z )  / GetGridSize() ) * GetGridSize() - BrushLocation.Z;
 				}
 
 				// If the snapping resulted in an off plane polygon, triangulate it to compensate.
diff --git a/Engine/Source/Editor/UnrealEd/Private/EditorComponents.cpp b/Engine/Source/Editor/UnrealEd/Private/EditorComponents.cpp
--- a/Engine/Source/Editor/UnrealEd/Private/EditorComponents.cpp
+++ b/Engine/Source/Editor/UnrealEd/Private/EditorComponents.cpp
@@ -548,8 +548,8 @@ void FEditorCommonDrawHelper::DrawGridSection(float ViewportGridY,FVector* A,FVe
 	FLinearColor MinorColor = FMath::Lerp(Background, FLinearColor::White, 0.02f);
 
 	const FMatrix InvViewProjMatrix = View->ViewMatrices.ProjMatrix.Inverse() * View->ViewMatrices.ViewMatrix.Inverse();
-	int32 FirstLine = FMath::Trunc(InvViewProjMatrix.TransformPosition(FVector(-1, -1, 0.5f)).Component(Axis) / ViewportGridY);
-	int32 LastLine = FMath::Trunc(InvViewProjMatrix.TransformPosition(FVector(+1, +1, 0.5f)).Component(Axis) / ViewportGridY);
+	int32 FirstLine = FMath::TruncToInt(InvViewProjMatrix.TransformPosition(FVector(-1, -1, 0.5f)).Component(Axis) / ViewportGridY);
+	int32 LastLine = FMath::TruncToInt(InvViewProjMatrix.TransformPosition(FVector(+1, +1, 0.5f)).Component(Axis) / ViewportGridY);
 	if (FirstLine > LastLine)
 	{
 		Exchange(FirstLine, LastLine);
@@ -560,8 +560,8 @@ void FEditorCommonDrawHelper::DrawGridSection(float ViewportGridY,FVector* A,FVe
 	const int32 LastLineClamped = FMath::Min<int32>(LastLine + 1, +HALF_WORLD_MAX/ViewportGridY) / IncScale;
 	for( int32 LineIndex = FirstLineClamped; LineIndex <= LastLineClamped; ++LineIndex )
 	{
-		*AX = FPlatformMath::Trunc(LineIndex * ViewportGridY) * IncScale;
-		*BX = FPlatformMath::Trunc(LineIndex * ViewportGridY) * IncScale;
+		*AX = FPlatformMath::TruncToFloat(LineIndex * ViewportGridY) * IncScale;
+		*BX = FPlatformMath::TruncToFloat(LineIndex * ViewportGridY) * IncScale;
 
 		// Only minor lines fade out with ortho zoom distance.  Origin lines and major lines are drawn
 		// at 100% opacity, but with a brighter value
diff --git a/Engine/Source/Editor/UnrealEd/Private/EditorServer.cpp b/Engine/Source/Editor/UnrealEd/Private/EditorServer.cpp
--- a/Engine/Source/Editor/UnrealEd/Private/EditorServer.cpp
+++ b/Engine/Source/Editor/UnrealEd/Private/EditorServer.cpp
@@ -6166,7 +6166,7 @@ void UEditorEngine::AutoMergeStaticMeshes()
 		NewStaticMesh->LightMapCoordinateIndex = LightmapUVChannel;
 
 		// figure out how much to grow the lightmap resolution by, since it needs to be square, start by sqrt'ing the number
-		int32 LightmapMultiplier = FMath::Trunc(FMath::Sqrt(MergeComponents.Num()));
+		int32 LightmapMultiplier = FMath::TruncToInt(FMath::Sqrt(MergeComponents.Num()));
 
 		// increase the sqrt by 1 unless it was a perfect square
 		if (LightmapMultiplier * LightmapMultiplier != MergeComponents.Num())
diff --git a/Engine/Source/Editor/UnrealEd/Private/EditorViewportClient.cpp b/Engine/Source/Editor/UnrealEd/Private/EditorViewportClient.cpp
--- a/Engine/Source/Editor/UnrealEd/Private/EditorViewportClient.cpp
+++ b/Engine/Source/Editor/UnrealEd/Private/EditorViewportClient.cpp
@@ -2137,7 +2137,7 @@ FString UnrealUnitsToSiUnits(float UnrealUnits)
 	// Put it in mm to start off with
 	UnrealUnits *= 10.f;
 
-	const int32 OrderOfMagnitude = UnrealUnits > 0 ? FMath::Trunc(FMath::LogX(10.0f, UnrealUnits)) : 0;
+	const int32 OrderOfMagnitude = UnrealUnits > 0 ? FMath::TruncToInt(FMath::LogX(10.0f, UnrealUnits)) : 0;
 
 	// Get an exponent applied to anything >= 1,000,000,000mm (1000km)
 	const int32 Exponent = (OrderOfMagnitude - 6)  / 3;
@@ -2158,7 +2158,7 @@ FString UnrealUnitsToSiUnits(float UnrealUnits)
 	{
 		const int32 ScaledOrder = OrderOfMagnitude % (FMath::Max(OrderOfThousands, 1) * 3);
 		const float RoundingDivisor = FMath::Pow(10.f, ScaledOrder) / 10.f;
-		const int32 Rounded = FMath::Trunc(ScaledNumber / RoundingDivisor) * RoundingDivisor;
+		const int32 Rounded = FMath::TruncToInt(ScaledNumber / RoundingDivisor) * RoundingDivisor;
 		if (ScaledNumber - Rounded > KINDA_SMALL_NUMBER)
 		{
 			ScaledNumber = Rounded;
@@ -2179,7 +2179,7 @@ FString UnrealUnitsToSiUnits(float UnrealUnits)
 	}
 	else
 	{
-		return FString::Printf(TEXT("%s%d%s%s"), Approximation, FMath::Trunc(ScaledNumber), *ExponentString, UnitText[FMath::Min(OrderOfThousands, 2)]);
+		return FString::Printf(TEXT("%s%d%s%s"), Approximation, FMath::TruncToInt(ScaledNumber), *ExponentString, UnitText[FMath::Min(OrderOfThousands, 2)]);
 	}
 }
 
@@ -2189,7 +2189,7 @@ void FEditorViewportClient::DrawScaleUnits(FViewport* InViewport, FCanvas* Canva
 
 	// Find the closest power of ten to our target width
 	static const int32 ApproxTargetMarkerWidthPx = 100;
-	const float SegmentWidthUnits = UnitsPerPixel > 0 ? FMath::Pow(10.f, FMath::Round(FMath::LogX(10.f, UnitsPerPixel * ApproxTargetMarkerWidthPx))) : 0.f;
+	const float SegmentWidthUnits = UnitsPerPixel > 0 ? FMath::Pow(10.f, FMath::RoundToFloat(FMath::LogX(10.f, UnitsPerPixel * ApproxTargetMarkerWidthPx))) : 0.f;
 
 	const FString DisplayText = UnrealUnitsToSiUnits(SegmentWidthUnits);
 	
diff --git a/Engine/Source/Editor/UnrealEd/Private/Factories/EditorFactories.cpp b/Engine/Source/Editor/UnrealEd/Private/Factories/EditorFactories.cpp
--- a/Engine/Source/Editor/UnrealEd/Private/Factories/EditorFactories.cpp
+++ b/Engine/Source/Editor/UnrealEd/Private/Factories/EditorFactories.cpp
@@ -4737,10 +4737,10 @@ class FHDRExportHelper
 			int32 Exponent;
 			const float Scale	 = frexp(Primary, &Exponent) / Primary * 255.f;
 
-			ReturnColor.R = FMath::Clamp(FMath::Trunc((R* Scale) + Rand.GetFraction()), 0, 255);
-			ReturnColor.G = FMath::Clamp(FMath::Trunc((G* Scale) + Rand.GetFraction()), 0, 255);
-			ReturnColor.B = FMath::Clamp(FMath::Trunc((B* Scale) + Rand.GetFraction()), 0, 255);
-			ReturnColor.A = FMath::Clamp(FMath::Trunc(Exponent), -128, 127) + 128;
+			ReturnColor.R = FMath::Clamp(FMath::TruncToInt((R* Scale) + Rand.GetFraction()), 0, 255);
+			ReturnColor.G = FMath::Clamp(FMath::TruncToInt((G* Scale) + Rand.GetFraction()), 0, 255);
+			ReturnColor.B = FMath::Clamp(FMath::TruncToInt((B* Scale) + Rand.GetFraction()), 0, 255);
+			ReturnColor.A = FMath::Clamp(FMath::TruncToInt(Exponent), -128, 127) + 128;
 		}
 
 		return ReturnColor;
diff --git a/Engine/Source/Editor/UnrealEd/Private/Factories/TTFontImport.cpp b/Engine/Source/Editor/UnrealEd/Private/Factories/TTFontImport.cpp
--- a/Engine/Source/Editor/UnrealEd/Private/Factories/TTFontImport.cpp
+++ b/Engine/Source/Editor/UnrealEd/Private/Factories/TTFontImport.cpp
@@ -702,7 +702,7 @@ float FTextureAlphaToDistanceField::FBuildDistanceFieldTask::CalcSignedDistanceT
 		// on the edge.
 		if ( bFoundClosest && RequiredRadius >= InScanRadius )
 		{
-			RequiredRadius = FMath::Ceil( FMath::Sqrt(RingSize*RingSize*2) );
+			RequiredRadius = FMath::CeilToInt( FMath::Sqrt(RingSize*RingSize*2) );
 			RequiredRadius = FMath::Min(RequiredRadius, InScanRadius);
 		}
 	}
@@ -1003,7 +1003,7 @@ bool UTrueTypeFontFactory::CreateFontTexture(
 
 	for( int32 Page = 0; Page < NumResolutions; ++Page )
 	{
-		int32 nHeight = -FMath::Round( ResHeights[Page] * LogicalPPIYRatio );
+		int32 nHeight = -FMath::RoundToInt( ResHeights[Page] * LogicalPPIYRatio );
 		
 		// scale font height to generate high res bitmap based on scale factor
 		// this high res bitmap is later used to generate the downsampled distance field
@@ -1236,8 +1236,8 @@ bool UTrueTypeFontFactory::CreateFontTexture(
 			{
 				
 				// Make X and Y positions a multiple of the scale factor.
-				CharWidth = FMath::Round(CharWidth / static_cast<float>(ImportOptions->Data.DistanceFieldScaleFactor)) * ImportOptions->Data.DistanceFieldScaleFactor;
-				CharHeight = FMath::Round(CharHeight / static_cast<float>(ImportOptions->Data.DistanceFieldScaleFactor)) * ImportOptions->Data.DistanceFieldScaleFactor;
+				CharWidth = FMath::RoundToInt(CharWidth / static_cast<float>(ImportOptions->Data.DistanceFieldScaleFactor)) * ImportOptions->Data.DistanceFieldScaleFactor;
+				CharHeight = FMath::RoundToInt(CharHeight / static_cast<float>(ImportOptions->Data.DistanceFieldScaleFactor)) * ImportOptions->Data.DistanceFieldScaleFactor;
 			}
 
 			// If the character is bigger than our texture size, then this isn't going to work!  The user
@@ -1294,10 +1294,10 @@ bool UTrueTypeFontFactory::CreateFontTexture(
 			// scale character offset UVs back down based on scale factor
 			if (ImportOptions->Data.bUseDistanceFieldAlpha)
 			{
-				FontX = FMath::Round(FontX / (float)ImportOptions->Data.DistanceFieldScaleFactor);
-				FontY = FMath::Round(FontY / (float)ImportOptions->Data.DistanceFieldScaleFactor);
-				FontWidth = FMath::Round(FontWidth / (float)ImportOptions->Data.DistanceFieldScaleFactor);
-				FontHeight = FMath::Round(FontHeight / (float)ImportOptions->Data.DistanceFieldScaleFactor);
+				FontX = FMath::RoundToInt(FontX / (float)ImportOptions->Data.DistanceFieldScaleFactor);
+				FontY = FMath::RoundToInt(FontY / (float)ImportOptions->Data.DistanceFieldScaleFactor);
+				FontWidth = FMath::RoundToInt(FontWidth / (float)ImportOptions->Data.DistanceFieldScaleFactor);
+				FontHeight = FMath::RoundToInt(FontHeight / (float)ImportOptions->Data.DistanceFieldScaleFactor);
 			}
 			NewCharacterRef.StartU =
 				FMath::Clamp<int32>( FontX - ImportOptions->Data.ExtendBoxLeft,
@@ -1837,7 +1837,7 @@ bool UTrueTypeFontFactory::CreateFontTexture(
 
 	for( int32 Page = 0; Page < NumResolutions; ++Page )
 	{
-		int32 Height = FMath::Round( ResHeights[Page] );
+		int32 Height = FMath::RoundToInt( ResHeights[Page] );
 		
 		// scale font height to generate high res bitmap based on scale factor
 		// this high res bitmap is later used to generate the downsampled distance field
@@ -1956,8 +1956,8 @@ bool UTrueTypeFontFactory::CreateFontTexture(
 			{
 				
 				// Make X and Y positions a multiple of the scale factor.
-				CharWidth = FMath::Round(CharWidth / static_cast<float>(ImportOptions->Data.DistanceFieldScaleFactor)) * ImportOptions->Data.DistanceFieldScaleFactor;
-				CharHeight = FMath::Round(CharHeight / static_cast<float>(ImportOptions->Data.DistanceFieldScaleFactor)) * ImportOptions->Data.DistanceFieldScaleFactor;
+				CharWidth = FMath::RoundToInt(CharWidth / static_cast<float>(ImportOptions->Data.DistanceFieldScaleFactor)) * ImportOptions->Data.DistanceFieldScaleFactor;
+				CharHeight = FMath::RoundToInt(CharHeight / static_cast<float>(ImportOptions->Data.DistanceFieldScaleFactor)) * ImportOptions->Data.DistanceFieldScaleFactor;
 			}
 
 			// If the character is bigger than our texture size, then this isn't going to work!  The user
@@ -2007,10 +2007,10 @@ bool UTrueTypeFontFactory::CreateFontTexture(
 			// scale character offset UVs back down based on scale factor
 			if (ImportOptions->Data.bUseDistanceFieldAlpha)
 			{
-				FontX = FMath::Round(FontX / (float)ImportOptions->Data.DistanceFieldScaleFactor);
-				FontY = FMath::Round(FontY / (float)ImportOptions->Data.DistanceFieldScaleFactor);
-				FontWidth = FMath::Round(FontWidth / (float)ImportOptions->Data.DistanceFieldScaleFactor);
-				FontHeight = FMath::Round(FontHeight / (float)ImportOptions->Data.DistanceFieldScaleFactor);
+				FontX = FMath::RoundToInt(FontX / (float)ImportOptions->Data.DistanceFieldScaleFactor);
+				FontY = FMath::RoundToInt(FontY / (float)ImportOptions->Data.DistanceFieldScaleFactor);
+				FontWidth = FMath::RoundToInt(FontWidth / (float)ImportOptions->Data.DistanceFieldScaleFactor);
+				FontHeight = FMath::RoundToInt(FontHeight / (float)ImportOptions->Data.DistanceFieldScaleFactor);
 			}
 			NewCharacterRef.StartU =
 				FMath::Clamp<int32>( FontX - ImportOptions->Data.ExtendBoxLeft,
diff --git a/Engine/Source/Editor/UnrealEd/Private/Fbx/FbxAnimationExport.cpp b/Engine/Source/Editor/UnrealEd/Private/Fbx/FbxAnimationExport.cpp
--- a/Engine/Source/Editor/UnrealEd/Private/Fbx/FbxAnimationExport.cpp
+++ b/Engine/Source/Editor/UnrealEd/Private/Fbx/FbxAnimationExport.cpp
@@ -352,7 +352,7 @@ void FFbxExporter::ExportMatineeGroup(class AMatineeActor* MatineeActor, USkelet
 		if( NextUpdateTime <= 0.0f )
 		{
 			NextUpdateTime = UpdateFrequency;
-			GWarn->StatusUpdate( FMath::Round( SampleTime ), FMath::Round(MatineeLength), NSLOCTEXT("FbxExporter", "ExportingToFbxStatus", "Exporting to FBX") );
+			GWarn->StatusUpdate( FMath::RoundToInt( SampleTime ), FMath::RoundToInt(MatineeLength), NSLOCTEXT("FbxExporter", "ExportingToFbxStatus", "Exporting to FBX") );
 		}
 
 		// Add the animation data to the bone nodes
diff --git a/Engine/Source/Editor/UnrealEd/Private/Fbx/FbxMaterialImport.cpp b/Engine/Source/Editor/UnrealEd/Private/Fbx/FbxMaterialImport.cpp
--- a/Engine/Source/Editor/UnrealEd/Private/Fbx/FbxMaterialImport.cpp
+++ b/Engine/Source/Editor/UnrealEd/Private/Fbx/FbxMaterialImport.cpp
@@ -256,8 +256,8 @@ bool UnFbx::FFbxImporter::CreateAndLinkExpressionForMaterialProperty(
 						MaterialInput.Expression = UnrealTextureExpression;
 						UnrealTextureExpression->Texture = UnrealTexture;
 						UnrealTextureExpression->SamplerType = bSetupAsNormalMap ? SAMPLERTYPE_Normal : SAMPLERTYPE_Color;
-						UnrealTextureExpression->MaterialExpressionEditorX = FMath::Trunc(Location.X);
-						UnrealTextureExpression->MaterialExpressionEditorY = FMath::Trunc(Location.Y);
+						UnrealTextureExpression->MaterialExpressionEditorX = FMath::TruncToInt(Location.X);
+						UnrealTextureExpression->MaterialExpressionEditorY = FMath::TruncToInt(Location.Y);
 
 						// add/find UVSet and set it to the texture
 						FbxString UVSetName = FbxTexture->UVSet.Get();
@@ -272,8 +272,8 @@ bool UnFbx::FFbxImporter::CreateAndLinkExpressionForMaterialProperty(
 							MyCoordExpression->UTiling = ScaleU;
 							MyCoordExpression->VTiling = ScaleV;
 							UnrealTextureExpression->Coordinates.Expression = MyCoordExpression;
-							MyCoordExpression->MaterialExpressionEditorX = FMath::Trunc(Location.X+175);
-							MyCoordExpression->MaterialExpressionEditorY = FMath::Trunc(Location.Y);
+							MyCoordExpression->MaterialExpressionEditorX = FMath::TruncToInt(Location.X+175);
+							MyCoordExpression->MaterialExpressionEditorY = FMath::TruncToInt(Location.Y);
 
 						}
 
diff --git a/Engine/Source/Editor/UnrealEd/Private/GeomTools.cpp b/Engine/Source/Editor/UnrealEd/Private/GeomTools.cpp
--- a/Engine/Source/Editor/UnrealEd/Private/GeomTools.cpp
+++ b/Engine/Source/Editor/UnrealEd/Private/GeomTools.cpp
@@ -118,10 +118,10 @@ FClipSMVertex InterpolateVert(const FClipSMVertex& V0, const FClipSMVertex& V1,
 		Result.UVs[i] = FMath::Lerp(V0.UVs[i], V1.UVs[i], Alpha);
 	}
 	
-	Result.Color.R = FMath::Clamp( FMath::Trunc(FMath::Lerp(float(V0.Color.R), float(V1.Color.R), Alpha)), 0, 255 );
-	Result.Color.G = FMath::Clamp( FMath::Trunc(FMath::Lerp(float(V0.Color.G), float(V1.Color.G), Alpha)), 0, 255 );
-	Result.Color.B = FMath::Clamp( FMath::Trunc(FMath::Lerp(float(V0.Color.B), float(V1.Color.B), Alpha)), 0, 255 );
-	Result.Color.A = FMath::Clamp( FMath::Trunc(FMath::Lerp(float(V0.Color.A), float(V1.Color.A), Alpha)), 0, 255 );
+	Result.Color.R = FMath::Clamp( FMath::TruncToInt(FMath::Lerp(float(V0.Color.R), float(V1.Color.R), Alpha)), 0, 255 );
+	Result.Color.G = FMath::Clamp( FMath::TruncToInt(FMath::Lerp(float(V0.Color.G), float(V1.Color.G), Alpha)), 0, 255 );
+	Result.Color.B = FMath::Clamp( FMath::TruncToInt(FMath::Lerp(float(V0.Color.B), float(V1.Color.B), Alpha)), 0, 255 );
+	Result.Color.A = FMath::Clamp( FMath::TruncToInt(FMath::Lerp(float(V0.Color.A), float(V1.Color.A), Alpha)), 0, 255 );
 	return Result;
 }
 
diff --git a/Engine/Source/Editor/UnrealEd/Private/Lightmass/Lightmass.cpp b/Engine/Source/Editor/UnrealEd/Private/Lightmass/Lightmass.cpp
--- a/Engine/Source/Editor/UnrealEd/Private/Lightmass/Lightmass.cpp
+++ b/Engine/Source/Editor/UnrealEd/Private/Lightmass/Lightmass.cpp
@@ -2003,25 +2003,25 @@ void FLightmassExporter::WriteSceneSettings( Lightmass::FSceneFileHeader& Scene
 
 		float NumShadowRaysScale;
 		verify(GConfig->GetFloat(QualitySectionNames[QualityLevel], TEXT("NumShadowRaysScale"), NumShadowRaysScale, GLightmassIni));
-		Scene.ShadowSettings.NumShadowRays = FMath::Trunc(Scene.ShadowSettings.NumShadowRays * NumShadowRaysScale);
+		Scene.ShadowSettings.NumShadowRays = FMath::TruncToInt(Scene.ShadowSettings.NumShadowRays * NumShadowRaysScale);
 
 		float NumPenumbraShadowRaysScale;
 		verify(GConfig->GetFloat(QualitySectionNames[QualityLevel], TEXT("NumPenumbraShadowRaysScale"), NumPenumbraShadowRaysScale, GLightmassIni));
-		Scene.ShadowSettings.NumPenumbraShadowRays = FMath::Trunc(Scene.ShadowSettings.NumPenumbraShadowRays * NumPenumbraShadowRaysScale);
+		Scene.ShadowSettings.NumPenumbraShadowRays = FMath::TruncToInt(Scene.ShadowSettings.NumPenumbraShadowRays * NumPenumbraShadowRaysScale);
 
 		float ApproximateHighResTexelsPerMaxTransitionDistanceScale;
 		verify(GConfig->GetFloat(QualitySectionNames[QualityLevel], TEXT("ApproximateHighResTexelsPerMaxTransitionDistanceScale"), ApproximateHighResTexelsPerMaxTransitionDistanceScale, GLightmassIni));
-		Scene.ShadowSettings.ApproximateHighResTexelsPerMaxTransitionDistance = FMath::Trunc(Scene.ShadowSettings.ApproximateHighResTexelsPerMaxTransitionDistance * ApproximateHighResTexelsPerMaxTransitionDistanceScale);
+		Scene.ShadowSettings.ApproximateHighResTexelsPerMaxTransitionDistance = FMath::TruncToInt(Scene.ShadowSettings.ApproximateHighResTexelsPerMaxTransitionDistance * ApproximateHighResTexelsPerMaxTransitionDistanceScale);
 
 		verify(GConfig->GetInt(QualitySectionNames[QualityLevel], TEXT("MinDistanceFieldUpsampleFactor"), Scene.ShadowSettings.MinDistanceFieldUpsampleFactor, GLightmassIni));
 
 		float NumHemisphereSamplesScale;
 		verify(GConfig->GetFloat(QualitySectionNames[QualityLevel], TEXT("NumHemisphereSamplesScale"), NumHemisphereSamplesScale, GLightmassIni));
-		Scene.ImportanceTracingSettings.NumHemisphereSamples = FMath::Trunc(Scene.ImportanceTracingSettings.NumHemisphereSamples * NumHemisphereSamplesScale);
+		Scene.ImportanceTracingSettings.NumHemisphereSamples = FMath::TruncToInt(Scene.ImportanceTracingSettings.NumHemisphereSamples * NumHemisphereSamplesScale);
 
 		float NumImportanceSearchPhotonsScale;
 		verify(GConfig->GetFloat(QualitySectionNames[QualityLevel], TEXT("NumImportanceSearchPhotonsScale"), NumImportanceSearchPhotonsScale, GLightmassIni));
-		Scene.PhotonMappingSettings.NumImportanceSearchPhotons = FMath::Trunc(Scene.PhotonMappingSettings.NumImportanceSearchPhotons * NumImportanceSearchPhotonsScale);
+		Scene.PhotonMappingSettings.NumImportanceSearchPhotons = FMath::TruncToInt(Scene.PhotonMappingSettings.NumImportanceSearchPhotons * NumImportanceSearchPhotonsScale);
 
 		float NumDirectPhotonsScale;
 		verify(GConfig->GetFloat(QualitySectionNames[QualityLevel], TEXT("NumDirectPhotonsScale"), NumDirectPhotonsScale, GLightmassIni));
@@ -3130,10 +3130,10 @@ void FLightmassProcessor::ApplyPrecomputedVisibility()
 			const FUncompressedPrecomputedVisibilityCell& CurrentCell = CombinedPrecomputedVisibilityCells[CellIndex];
 
 			const float FloatOffsetX = (CurrentCell.Bounds.Min.X - CellBucketOriginXY.X + .5f * CellSize) / CellSize;
-			// FMath::Trunc rounds toward 0, we want to always round down
-			const int32 BucketIndexX = FMath::Abs((FMath::Trunc(FloatOffsetX) - (FloatOffsetX < 0.0f ? 1 : 0)) / CellBucketSize % NumCellBuckets);
+			// FMath::TruncToInt rounds toward 0, we want to always round down
+			const int32 BucketIndexX = FMath::Abs((FMath::TruncToInt(FloatOffsetX) - (FloatOffsetX < 0.0f ? 1 : 0)) / CellBucketSize % NumCellBuckets);
 			const float FloatOffsetY = (CurrentCell.Bounds.Min.Y - CellBucketOriginXY.Y + .5f * CellSize) / CellSize;
-			const int32 BucketIndexY = FMath::Abs((FMath::Trunc(FloatOffsetY) - (FloatOffsetY < 0.0f ? 1 : 0)) / CellBucketSize % NumCellBuckets);
+			const int32 BucketIndexY = FMath::Abs((FMath::TruncToInt(FloatOffsetY) - (FloatOffsetY < 0.0f ? 1 : 0)) / CellBucketSize % NumCellBuckets);
 
 			const int32 BucketIndex = BucketIndexY * CellBucketSize + BucketIndexX;
 			CellRenderingBuckets[BucketIndex].Add(&CurrentCell);
@@ -3931,8 +3931,8 @@ bool FLightmassProcessor::ImportTextureMapping(int32 Channel, FTextureMappingImp
 		LightMapTypeModifier = NUM_LQ_LIGHTMAP_COEF;
 	}
 
-	int32 WastedMemory = FMath::Trunc(MemoryAmount * BytesPerPixel * MIP_FACTOR * LightMapTypeModifier);
-	int32 TotalMemory = FMath::Trunc(TotalMemoryAmount * BytesPerPixel * MIP_FACTOR * LightMapTypeModifier);
+	int32 WastedMemory = FMath::TruncToInt(MemoryAmount * BytesPerPixel * MIP_FACTOR * LightMapTypeModifier);
+	int32 TotalMemory = FMath::TruncToInt(TotalMemoryAmount * BytesPerPixel * MIP_FACTOR * LightMapTypeModifier);
 	
 	FStatsViewerModule& StatsViewerModule = FModuleManager::Get().LoadModuleChecked<FStatsViewerModule>(TEXT("StatsViewer"));
 	ULightingBuildInfo* LightingBuildInfo = ConstructObject<ULightingBuildInfo>(ULightingBuildInfo::StaticClass());
diff --git a/Engine/Source/Editor/UnrealEd/Private/Lightmass/LightmassRender.cpp b/Engine/Source/Editor/UnrealEd/Private/Lightmass/LightmassRender.cpp
--- a/Engine/Source/Editor/UnrealEd/Private/Lightmass/LightmassRender.cpp
+++ b/Engine/Source/Editor/UnrealEd/Private/Lightmass/LightmassRender.cpp
@@ -646,7 +646,7 @@ public:
 			OutSizeY = 0;
 			return false;
 		}
-		OutSizeX = OutSizeY = FMath::Trunc(GlobalSize * SizeScale);
+		OutSizeX = OutSizeY = FMath::TruncToInt(GlobalSize * SizeScale);
 		return true;
 	}
 
diff --git a/Engine/Source/Editor/UnrealEd/Private/PackageAutoSaver.cpp b/Engine/Source/Editor/UnrealEd/Private/PackageAutoSaver.cpp
--- a/Engine/Source/Editor/UnrealEd/Private/PackageAutoSaver.cpp
+++ b/Engine/Source/Editor/UnrealEd/Private/PackageAutoSaver.cpp
@@ -349,7 +349,7 @@ int32 FPackageAutoSaver::GetTimeTillAutoSave(const bool bIgnoreCanAutoSave) cons
 	int32 Result = -1;
 	if (bIgnoreCanAutoSave || CanAutoSave())
 	{
-		Result = FMath::Ceil(GetDefault<UEditorLoadingSavingSettings>()->AutoSaveTimeMinutes * 60.0f - AutoSaveCount);
+		Result = FMath::CeilToInt(GetDefault<UEditorLoadingSavingSettings>()->AutoSaveTimeMinutes * 60.0f - AutoSaveCount);
 	}
 
 	return Result;
diff --git a/Engine/Source/Editor/UnrealEd/Private/PlayLevel.cpp b/Engine/Source/Editor/UnrealEd/Private/PlayLevel.cpp
--- a/Engine/Source/Editor/UnrealEd/Private/PlayLevel.cpp
+++ b/Engine/Source/Editor/UnrealEd/Private/PlayLevel.cpp
@@ -2359,8 +2359,8 @@ UWorld* UEditorEngine::CreatePlayInEditorWorld(FWorldContext &PieWorldContext, b
 
 							if (index <= 0)
 							{
-								LevelEditorPlaySettings->NewWindowPosition.X = FPlatformMath::Round(PIEWindowPos.X);
-								LevelEditorPlaySettings->NewWindowPosition.Y = FPlatformMath::Round(PIEWindowPos.Y);
+								LevelEditorPlaySettings->NewWindowPosition.X = FPlatformMath::RoundToInt(PIEWindowPos.X);
+								LevelEditorPlaySettings->NewWindowPosition.Y = FPlatformMath::RoundToInt(PIEWindowPos.Y);
 							}
 							else
 							{
diff --git a/Engine/Source/Editor/UnrealEd/Private/SColorGradientEditor.cpp b/Engine/Source/Editor/UnrealEd/Private/SColorGradientEditor.cpp
--- a/Engine/Source/Editor/UnrealEd/Private/SColorGradientEditor.cpp
+++ b/Engine/Source/Editor/UnrealEd/Private/SColorGradientEditor.cpp
@@ -142,7 +142,7 @@ int32 SColorGradientEditor::OnPaint( const FGeometry& AllottedGeometry, const FS
 
 		// The start and end location in slate units of the area to draw
 		int32 Start = 0;
-		int32 Finish = FMath::Trunc( AllottedGeometry.Size.X );
+		int32 Finish = FMath::TruncToInt( AllottedGeometry.Size.X );
 
 		TArray<FSlateGradientStop> Stops;
 
diff --git a/Engine/Source/Editor/UnrealEd/Private/StaticLightingSystem/StaticLightingDebug.cpp b/Engine/Source/Editor/UnrealEd/Private/StaticLightingSystem/StaticLightingDebug.cpp
--- a/Engine/Source/Editor/UnrealEd/Private/StaticLightingSystem/StaticLightingDebug.cpp
+++ b/Engine/Source/Editor/UnrealEd/Private/StaticLightingSystem/StaticLightingDebug.cpp
@@ -102,8 +102,8 @@ static bool UpdateSelectedTexel(
 			int32 LightmapIndex = Lightmap2D->AllowsHighQualityLightmaps() ? 0 : 1;
 			UTexture2D* CurrentLightmap = Lightmap2D->GetTexture( LightmapIndex );
 			// UV's in the lightmap atlas
-			NewSelectedTexel.LightmapX = FMath::Trunc(LightmapUV.X * CurrentLightmap->GetSizeX());
-			NewSelectedTexel.LightmapY = FMath::Trunc(LightmapUV.Y * CurrentLightmap->GetSizeY());
+			NewSelectedTexel.LightmapX = FMath::TruncToInt(LightmapUV.X * CurrentLightmap->GetSizeX());
+			NewSelectedTexel.LightmapY = FMath::TruncToInt(LightmapUV.Y * CurrentLightmap->GetSizeY());
 			// Write the selection color to the selected lightmap texel
 			WriteTexel(CurrentLightmap, NewSelectedTexel.LightmapX, NewSelectedTexel.LightmapY, GTexelSelectionColor, NewSelectedTexel.OriginalColor);
 
@@ -243,8 +243,8 @@ void SetDebugLightmapSample(TArray<UActorComponent*>* Components, UModel* Model,
 							PaddedSizeY -= 2;
 						}
 
-						const int32 LocalX = FMath::Trunc(InterpolatedUV.X * PaddedSizeX);
-						const int32 LocalY = FMath::Trunc(InterpolatedUV.Y * PaddedSizeY);
+						const int32 LocalX = FMath::TruncToInt(InterpolatedUV.X * PaddedSizeX);
+						const int32 LocalY = FMath::TruncToInt(InterpolatedUV.Y * PaddedSizeY);
 						if (LocalX < 0 || LocalX >= PaddedSizeX
 							|| LocalY < 0 || LocalY >= PaddedSizeY)
 						{
@@ -296,8 +296,8 @@ void SetDebugLightmapSample(TArray<UActorComponent*>* Components, UModel* Model,
 					PaddedSizeY -= 2;
 				}
 
-				const int32 LocalX = FMath::Trunc(InterpolatedUV.X * PaddedSizeX);
-				const int32 LocalY = FMath::Trunc(InterpolatedUV.Y * PaddedSizeY);
+				const int32 LocalX = FMath::TruncToInt(InterpolatedUV.X * PaddedSizeX);
+				const int32 LocalY = FMath::TruncToInt(InterpolatedUV.Y * PaddedSizeY);
 				if (LocalX < 0 || LocalX >= PaddedSizeX
 					|| LocalY < 0 || LocalY >= PaddedSizeY)
 				{
@@ -405,8 +405,8 @@ void SetDebugLightmapSample(TArray<UActorComponent*>* Components, UModel* Model,
 
 							// Apply the transform to the intersection position to find the local texel coordinates
 							const FVector4 StaticLightingTextureCoordinate = WorldToMap.TransformPosition(ClickLocation);
-							const int32 LocalX = FMath::Trunc(StaticLightingTextureCoordinate.X * PaddedSizeX);
-							const int32 LocalY = FMath::Trunc(StaticLightingTextureCoordinate.Y * PaddedSizeY);
+							const int32 LocalX = FMath::TruncToInt(StaticLightingTextureCoordinate.X * PaddedSizeX);
+							const int32 LocalY = FMath::TruncToInt(StaticLightingTextureCoordinate.Y * PaddedSizeY);
 							check(LocalX >= 0 && LocalX < PaddedSizeX && LocalY >= 0 && LocalY < PaddedSizeY);
 
 							bFoundLightmapSample = UpdateSelectedTexel(
diff --git a/Engine/Source/Editor/UnrealEd/Private/ThumbnailRendering/DefaultSizedThumbnailRenderer.cpp b/Engine/Source/Editor/UnrealEd/Private/ThumbnailRendering/DefaultSizedThumbnailRenderer.cpp
--- a/Engine/Source/Editor/UnrealEd/Private/ThumbnailRendering/DefaultSizedThumbnailRenderer.cpp
+++ b/Engine/Source/Editor/UnrealEd/Private/ThumbnailRendering/DefaultSizedThumbnailRenderer.cpp
@@ -9,6 +9,6 @@ UDefaultSizedThumbnailRenderer::UDefaultSizedThumbnailRenderer(const class FPost
 
 void UDefaultSizedThumbnailRenderer::GetThumbnailSize(UObject*, float Zoom, uint32& OutWidth, uint32& OutHeight) const
 {
-	OutWidth = FMath::Trunc(DefaultSizeX * Zoom);
-	OutHeight = FMath::Trunc(DefaultSizeY * Zoom);
+	OutWidth = FMath::TruncToInt(DefaultSizeX * Zoom);
+	OutHeight = FMath::TruncToInt(DefaultSizeY * Zoom);
 }
diff --git a/Engine/Source/Editor/UnrealEd/Private/ThumbnailRendering/FontThumbnailRenderer.cpp b/Engine/Source/Editor/UnrealEd/Private/ThumbnailRendering/FontThumbnailRenderer.cpp
--- a/Engine/Source/Editor/UnrealEd/Private/ThumbnailRendering/FontThumbnailRenderer.cpp
+++ b/Engine/Source/Editor/UnrealEd/Private/ThumbnailRendering/FontThumbnailRenderer.cpp
@@ -16,8 +16,8 @@ void UFontThumbnailRenderer::GetThumbnailSize(UObject* Object, float Zoom,uint32
 	{
 		// Get the texture interface for the font text
 		UTexture2D* Tex = Font->Textures[0];
-		OutWidth = FMath::Trunc(Zoom * (float)Tex->GetSurfaceWidth());
-		OutHeight = FMath::Trunc(Zoom * (float)Tex->GetSurfaceHeight());
+		OutWidth = FMath::TruncToInt(Zoom * (float)Tex->GetSurfaceWidth());
+		OutHeight = FMath::TruncToInt(Zoom * (float)Tex->GetSurfaceHeight());
 	}
 	else
 	{
diff --git a/Engine/Source/Editor/UnrealEd/Private/ThumbnailRendering/ParticleSystemThumbnailRenderer.cpp b/Engine/Source/Editor/UnrealEd/Private/ThumbnailRendering/ParticleSystemThumbnailRenderer.cpp
--- a/Engine/Source/Editor/UnrealEd/Private/ThumbnailRendering/ParticleSystemThumbnailRenderer.cpp
+++ b/Engine/Source/Editor/UnrealEd/Private/ThumbnailRendering/ParticleSystemThumbnailRenderer.cpp
@@ -36,8 +36,8 @@ void UParticleSystemThumbnailRenderer::GetThumbnailSize(UObject* Object, float Z
 			(PSys->ThumbnailImage) || 
 			(NoImage))
 		{
-			OutWidth = FMath::Trunc(1024 * Zoom);
-			OutHeight = FMath::Trunc(1024 * Zoom);
+			OutWidth = FMath::TruncToInt(1024 * Zoom);
+			OutHeight = FMath::TruncToInt(1024 * Zoom);
 		}
 		else
 		{
diff --git a/Engine/Source/Editor/UnrealEd/Private/ThumbnailRendering/TextureThumbnailRenderer.cpp b/Engine/Source/Editor/UnrealEd/Private/ThumbnailRendering/TextureThumbnailRenderer.cpp
--- a/Engine/Source/Editor/UnrealEd/Private/ThumbnailRendering/TextureThumbnailRenderer.cpp
+++ b/Engine/Source/Editor/UnrealEd/Private/ThumbnailRendering/TextureThumbnailRenderer.cpp
@@ -25,8 +25,8 @@ void UTextureThumbnailRenderer::GetThumbnailSize(UObject* Object, float Zoom, ui
 
 	if (Texture != nullptr)
 	{
-		OutWidth = FMath::Trunc(Zoom * (float)Texture->GetSurfaceWidth());
-		OutHeight = FMath::Trunc(Zoom * (float)Texture->GetSurfaceHeight());
+		OutWidth = FMath::TruncToInt(Zoom * (float)Texture->GetSurfaceWidth());
+		OutHeight = FMath::TruncToInt(Zoom * (float)Texture->GetSurfaceHeight());
 	}
 	else
 	{
diff --git a/Engine/Source/Editor/UnrealEd/Private/UnrealWidget.cpp b/Engine/Source/Editor/UnrealEd/Private/UnrealWidget.cpp
--- a/Engine/Source/Editor/UnrealEd/Private/UnrealWidget.cpp
+++ b/Engine/Source/Editor/UnrealEd/Private/UnrealWidget.cpp
@@ -78,8 +78,8 @@ void FWidget::DrawHUD (FCanvas* Canvas)
 {
 	if (HUDString.Len())
 	{
-		int32 StringPosX = FMath::Floor(HUDInfoPos.X);
-		int32 StringPosY = FMath::Floor(HUDInfoPos.Y);
+		int32 StringPosX = FMath::FloorToInt(HUDInfoPos.X);
+		int32 StringPosY = FMath::FloorToInt(HUDInfoPos.Y);
 
 		//measure string size
 		int32 StringSizeX, StringSizeY;
@@ -1465,7 +1465,7 @@ void FWidget::DrawThickArc (const FThickArcParams& InParams, const FVector& Axis
 	const int32 CircleSides = (GetDefault<ULevelEditorViewportSettings>()->TransformWidgetSizeAdjustment > 0) 
 		? AXIS_CIRCLE_SIDES + (GetDefault<ULevelEditorViewportSettings>()->TransformWidgetSizeAdjustment / 5)
 		: AXIS_CIRCLE_SIDES;
-	const int32 NumPoints = FMath::Trunc(CircleSides * (InEndAngle-InStartAngle)/(PI/2)) + 1;
+	const int32 NumPoints = FMath::TruncToInt(CircleSides * (InEndAngle-InStartAngle)/(PI/2)) + 1;
 
 	FColor TriangleColor = InColor;
 	FColor RingColor = InColor;
diff --git a/Engine/Source/Editor/WorldBrowser/Private/LevelCollectionModel.cpp b/Engine/Source/Editor/WorldBrowser/Private/LevelCollectionModel.cpp
--- a/Engine/Source/Editor/WorldBrowser/Private/LevelCollectionModel.cpp
+++ b/Engine/Source/Editor/WorldBrowser/Private/LevelCollectionModel.cpp
@@ -733,8 +733,8 @@ void FLevelCollectionModel::UpdateAllLevels()
 	
 	// Update world size
 	FBox WorldBounds = GetLevelsBoundingBox(AllLevelsList, false);
-	WorldSize.X = FMath::Round(WorldBounds.GetSize().X);
-	WorldSize.Y = FMath::Round(WorldBounds.GetSize().Y);
+	WorldSize.X = FMath::RoundToInt(WorldBounds.GetSize().X);
+	WorldSize.Y = FMath::RoundToInt(WorldBounds.GetSize().Y);
 }
 
 void FLevelCollectionModel::RedrawAllLevels()
diff --git a/Engine/Source/Editor/WorldBrowser/Private/SWorldMainView.cpp b/Engine/Source/Editor/WorldBrowser/Private/SWorldMainView.cpp
--- a/Engine/Source/Editor/WorldBrowser/Private/SWorldMainView.cpp
+++ b/Engine/Source/Editor/WorldBrowser/Private/SWorldMainView.cpp
@@ -604,7 +604,7 @@ FString SWorldMainView::GetCurrentLevelText() const
 FString SWorldMainView::GetMouseLocationText() const
 {
 	FVector2D MouseLocation = LevelsGridView->GetMouseWorldLocation();
-	return FString::Printf(TEXT("%d, %d"), FMath::Round(MouseLocation.X), FMath::Round(MouseLocation.Y));
+	return FString::Printf(TEXT("%d, %d"), FMath::RoundToInt(MouseLocation.X), FMath::RoundToInt(MouseLocation.Y));
 }
 
 FString	SWorldMainView::GetMarqueeSelectionSizeText() const
@@ -612,7 +612,7 @@ FString	SWorldMainView::GetMarqueeSelectionSizeText() const
 	FVector2D MarqueeSize = LevelsGridView->GetMarqueeWorldSize();
 	if (MarqueeSize.Size() > 0)
 	{
-		return FString::Printf(TEXT("%d x %d"), FMath::Round(MarqueeSize.X), FMath::Round(MarqueeSize.Y));
+		return FString::Printf(TEXT("%d x %d"), FMath::RoundToInt(MarqueeSize.X), FMath::RoundToInt(MarqueeSize.Y));
 	}
 	else
 	{
diff --git a/Engine/Source/Editor/WorldBrowser/Private/StreamingLevels/StreamingLevelCustomization.cpp b/Engine/Source/Editor/WorldBrowser/Private/StreamingLevels/StreamingLevelCustomization.cpp
--- a/Engine/Source/Editor/WorldBrowser/Private/StreamingLevels/StreamingLevelCustomization.cpp
+++ b/Engine/Source/Editor/WorldBrowser/Private/StreamingLevels/StreamingLevelCustomization.cpp
@@ -202,7 +202,7 @@ TOptional<int32> FStreamingLevelCustomization::GetLevelRotation() const
 	FQuat RotQ;
 	if (LevelRotationProperty->GetValue(RotQ) == FPropertyAccess::Success)
 	{
-		int32 YawValue = FMath::Round(RotQ.Rotator().Yaw);
+		int32 YawValue = FMath::RoundToInt(RotQ.Rotator().Yaw);
 		return YawValue < 0 ? (YawValue + 360) : YawValue;
 	}
 
diff --git a/Engine/Source/Editor/WorldBrowser/Private/Tiles/WorldTileModel.cpp b/Engine/Source/Editor/WorldBrowser/Private/Tiles/WorldTileModel.cpp
--- a/Engine/Source/Editor/WorldBrowser/Private/Tiles/WorldTileModel.cpp
+++ b/Engine/Source/Editor/WorldBrowser/Private/Tiles/WorldTileModel.cpp
@@ -811,7 +811,7 @@ FText FWorldTileModel::GetPositionText() const
 FText FWorldTileModel::GetBoundsExtentText() const
 {
 	FVector2D Size = GetLevelSize2D();
-	return FText::FromString(FString::Printf(TEXT("%d, %d"), FMath::Round(Size.X*0.5f), FMath::Round(Size.Y*0.5f)));
+	return FText::FromString(FString::Printf(TEXT("%d, %d"), FMath::RoundToInt(Size.X*0.5f), FMath::RoundToInt(Size.Y*0.5f)));
 }
 
 FText FWorldTileModel::GetLevelLayerNameText() const
diff --git a/Engine/Source/Programs/UnrealHeaderTool/Private/ParserHelper.h b/Engine/Source/Programs/UnrealHeaderTool/Private/ParserHelper.h
--- a/Engine/Source/Programs/UnrealHeaderTool/Private/ParserHelper.h
+++ b/Engine/Source/Programs/UnrealHeaderTool/Private/ParserHelper.h
@@ -844,12 +844,12 @@ public:
 			I = Byte;
 			return 1;
 		}
-		else if( TokenType==TOKEN_Const && Type==CPT_Float && Float==FMath::Trunc(Float))
+		else if( TokenType==TOKEN_Const && Type==CPT_Float && Float==FMath::TruncToInt(Float))
 		{
 			I = (int32) Float;
 			return 1;
 		}
-		else if( TokenType==TOKEN_Const && Type==CPT_Double && Float==FMath::Trunc(Double))
+		else if (TokenType == TOKEN_Const && Type == CPT_Double && Float == FMath::TruncToInt(Double))
 		{
 			I = (int32) Double;
 			return 1;
diff --git a/Engine/Source/Programs/UnrealLightmass/Private/CPUSolver/CPUSolver.cpp b/Engine/Source/Programs/UnrealLightmass/Private/CPUSolver/CPUSolver.cpp
--- a/Engine/Source/Programs/UnrealLightmass/Private/CPUSolver/CPUSolver.cpp
+++ b/Engine/Source/Programs/UnrealLightmass/Private/CPUSolver/CPUSolver.cpp
@@ -172,7 +172,7 @@ void BuildStaticLighting(const FGuid& SceneGuid, int32 NumThreads, bool bDumpTex
 	GStatistics.ImportTimeStart = FPlatformTime::Seconds();
 	NSwarm::FSwarmInterface::Initialize(*(FString(FPlatformProcess::BaseDir()) + TEXT("..\\DotNET\\SwarmInterface.dll")));
 	check(&NSwarm::FSwarmInterface::Get() != NULL);
-	GSwarm = new FLightmassSwarm( NSwarm::FSwarmInterface::Get(), SceneGuid, FMath::Trunc(GNumTasksPerThreadPrefetch*NumThreads) );
+	GSwarm = new FLightmassSwarm( NSwarm::FSwarmInterface::Get(), SceneGuid, FMath::TruncToInt(GNumTasksPerThreadPrefetch*NumThreads) );
 	GSwarm->SendMessage( NSwarm::FTimingMessage( NSwarm::PROGSTATE_BeginJob, -1 ) );
 
 	FLightmassImporter Importer( GSwarm );
diff --git a/Engine/Source/Programs/UnrealLightmass/Private/ImportExport/LightmassScene.cpp b/Engine/Source/Programs/UnrealLightmass/Private/ImportExport/LightmassScene.cpp
--- a/Engine/Source/Programs/UnrealLightmass/Private/ImportExport/LightmassScene.cpp
+++ b/Engine/Source/Programs/UnrealLightmass/Private/ImportExport/LightmassScene.cpp
@@ -593,13 +593,13 @@ int32 FDirectionalLight::GetNumDirectPhotons(float DirectPhotonDensity) const
 		const float ImportanceDiskAreaMillions = (float)PI * FMath::Square(ImportanceBounds.SphereRadius) / 1000000.0f;
 		checkSlow(SceneBounds.SphereRadius > ImportanceBounds.SphereRadius);
 		const float OutsideImportanceDiskAreaMillions = (float)PI * (FMath::Square(SceneBounds.SphereRadius) - FMath::Square(ImportanceBounds.SphereRadius)) / 1000000.0f;
-		NumDirectPhotons = FMath::Trunc(ImportanceDiskAreaMillions * DirectPhotonDensity + OutsideImportanceDiskAreaMillions * OutsideImportanceVolumeDensity);
+		NumDirectPhotons = FMath::TruncToInt(ImportanceDiskAreaMillions * DirectPhotonDensity + OutsideImportanceDiskAreaMillions * OutsideImportanceVolumeDensity);
 	}
 	else
 	{
 		// Gather enough photons to meet DirectPhotonDensity everywhere in the scene
 		const float SceneDiskAreaMillions = (float)PI * FMath::Square(SceneBounds.SphereRadius) / 1000000.0f;
-		NumDirectPhotons = FMath::Trunc(SceneDiskAreaMillions * DirectPhotonDensity);
+		NumDirectPhotons = FMath::TruncToInt(SceneDiskAreaMillions * DirectPhotonDensity);
 	}
 	return NumDirectPhotons == appTruncErrorCode ? INT_MAX : NumDirectPhotons;
 }
@@ -757,7 +757,7 @@ void FDirectionalLight::SampleDirection(
 	checkSlow(IndirectPathRays.Num() > 0);
 
 	const FVector2D DiskPosition2D = GetUniformUnitDiskPosition(RandomStream);
-	const int32 RayIndex = FMath::Trunc(RandomStream.GetFraction() * IndirectPathRays.Num());
+	const int32 RayIndex = FMath::TruncToInt(RandomStream.GetFraction() * IndirectPathRays.Num());
 	checkSlow(RayIndex >= 0 && RayIndex < IndirectPathRays.Num());
 
 	// Create the ray using a disk centered at the scene's origin, whose radius is the size of the scene
@@ -780,8 +780,8 @@ void FDirectionalLight::SampleDirection(
 	checkSlow(SampleLightSurfacePosition.X >= GridCenter.X - GridExtent && SampleLightSurfacePosition.X <= GridCenter.X + GridExtent);
 	checkSlow(SampleLightSurfacePosition.Y >= GridCenter.Y - GridExtent && SampleLightSurfacePosition.Y <= GridCenter.Y + GridExtent);
 	// Calculate the cell indices that the generated position falls into
-	const int32 CellX = FMath::Clamp(FMath::Trunc(GridSize * (SampleLightSurfacePosition.X - GridCenter.X + GridExtent) / (2.0f * GridExtent)), 0, GridSize - 1);
-	const int32 CellY = FMath::Clamp(FMath::Trunc(GridSize * (SampleLightSurfacePosition.Y - GridCenter.Y + GridExtent) / (2.0f * GridExtent)), 0, GridSize - 1);
+	const int32 CellX = FMath::Clamp(FMath::TruncToInt(GridSize * (SampleLightSurfacePosition.X - GridCenter.X + GridExtent) / (2.0f * GridExtent)), 0, GridSize - 1);
+	const int32 CellY = FMath::Clamp(FMath::TruncToInt(GridSize * (SampleLightSurfacePosition.Y - GridCenter.Y + GridExtent) / (2.0f * GridExtent)), 0, GridSize - 1);
 	const TArray<int32>& CurrentGridCell = PathRayGrid[CellY * GridSize + CellX];
 	// The cell containing the sample position must contain at least the index of the path used to generate this sample position
 	checkSlow(CurrentGridCell.Num() > 0);
@@ -878,7 +878,7 @@ int32 FPointLight::GetNumDirectPhotons(float DirectPhotonDensity) const
 {
 	// Gather enough photons to meet DirectPhotonDensity at the influence radius of the point light.
 	const float InfluenceSphereSurfaceAreaMillions = 4.0f * (float)PI * FMath::Square(Radius) / 1000000.0f;
-	const int32 NumDirectPhotons = FMath::Trunc(InfluenceSphereSurfaceAreaMillions * DirectPhotonDensity);
+	const int32 NumDirectPhotons = FMath::TruncToInt(InfluenceSphereSurfaceAreaMillions * DirectPhotonDensity);
 	return NumDirectPhotons == appTruncErrorCode ? INT_MAX : NumDirectPhotons;
 }
 
@@ -1013,7 +1013,7 @@ void FPointLight::SampleDirection(
 {
 	checkSlow(IndirectPathRays.Num() > 0);
 	// Pick an indirect path ray with uniform probability
-	const int32 RayIndex = FMath::Trunc(RandomStream.GetFraction() * IndirectPathRays.Num());
+	const int32 RayIndex = FMath::TruncToInt(RandomStream.GetFraction() * IndirectPathRays.Num());
 	checkSlow(RayIndex >= 0 && RayIndex < IndirectPathRays.Num());
 
 	const FVector4 PathRayDirection = IndirectPathRays[RayIndex].UnitDirection;
@@ -1227,7 +1227,7 @@ int32 FSpotLight::GetNumDirectPhotons(float DirectPhotonDensity) const
 	// Find the fraction of the sphere's surface area that is inside the cone
 	const float ConeSurfaceAreaSphereFraction = ConeSolidAngle / (4.0f * (float)PI);
 	// Gather enough photons to meet DirectPhotonDensity on the spherical cap at the influence radius of the spot light.
-	const int32 NumDirectPhotons = FMath::Trunc(InfluenceSphereSurfaceAreaMillions * ConeSurfaceAreaSphereFraction * DirectPhotonDensity);
+	const int32 NumDirectPhotons = FMath::TruncToInt(InfluenceSphereSurfaceAreaMillions * ConeSurfaceAreaSphereFraction * DirectPhotonDensity);
 	return NumDirectPhotons == appTruncErrorCode ? INT_MAX : NumDirectPhotons;
 }
 
@@ -1347,7 +1347,7 @@ int32 FMeshAreaLight::GetNumDirectPhotons(float DirectPhotonDensity) const
 	// Clamp the influence radius to the importance or scene radius for the purposes of emitting photons
 	// This prevents huge mesh area lights from emitting more photons than are needed
 	const float InfluenceSphereSurfaceAreaMillions = 4.0f * (float)PI * FMath::Square(FMath::Min(ImportanceBounds.SphereRadius, InfluenceRadius)) / 1000000.0f;
-	const int32 NumDirectPhotons = FMath::Trunc(InfluenceSphereSurfaceAreaMillions * DirectPhotonDensity);
+	const int32 NumDirectPhotons = FMath::TruncToInt(InfluenceSphereSurfaceAreaMillions * DirectPhotonDensity);
 	return NumDirectPhotons == appTruncErrorCode ? INT_MAX : NumDirectPhotons;
 }
 
@@ -1383,8 +1383,8 @@ void FMeshAreaLight::SetPrimitives(
 		}
 		const FVector2D SphericalCoordinates = FVector(CurrentPrimitive.SurfaceNormal).UnitCartesianToSpherical();
 		// Determine grid cell the primitive's normal falls into based on spherical coordinates
-		const int32 CacheX = FMath::Clamp(FMath::Trunc(SphericalCoordinates.X / (float)PI * MeshAreaLightGridSize), 0, MeshAreaLightGridSize - 1);
-		const int32 CacheY = FMath::Clamp(FMath::Trunc((SphericalCoordinates.Y + (float)PI) / (2 * (float)PI) * MeshAreaLightGridSize), 0, MeshAreaLightGridSize - 1);
+		const int32 CacheX = FMath::Clamp(FMath::TruncToInt(SphericalCoordinates.X / (float)PI * MeshAreaLightGridSize), 0, MeshAreaLightGridSize - 1);
+		const int32 CacheY = FMath::Clamp(FMath::TruncToInt((SphericalCoordinates.Y + (float)PI) / (2 * (float)PI) * MeshAreaLightGridSize), 0, MeshAreaLightGridSize - 1);
 		CachedPrimitiveNormals[CacheY * MeshAreaLightGridSize + CacheX].Add(CurrentPrimitive.SurfaceNormal);
 	}
 
@@ -1622,7 +1622,7 @@ void FMeshAreaLight::SampleDirection(
 {
 	checkSlow(IndirectPathRays.Num() > 0);
 	// Pick an indirect path ray with uniform probability
-	const int32 RayIndex = FMath::Trunc(RandomStream.GetFraction() * IndirectPathRays.Num());
+	const int32 RayIndex = FMath::TruncToInt(RandomStream.GetFraction() * IndirectPathRays.Num());
 	checkSlow(RayIndex >= 0 && RayIndex < IndirectPathRays.Num());
 	const FIndirectPathRay& ChosenPathRay = IndirectPathRays[RayIndex];
 	const FVector4 PathRayDirection = ChosenPathRay.UnitDirection;
@@ -1703,7 +1703,7 @@ void FMeshAreaLight::SampleLightSurface(FLMRandomStream& RandomStream, FLightSur
 	float FloatPrimitiveIndex;
 	// Pick a primitive with probability proportional to the primitive's fraction of the light's total surface area
 	Sample1dCDF(PrimitivePDFs, PrimitiveCDFs, UnnormalizedIntegral, RandomStream, PrimitivePDF, FloatPrimitiveIndex);
-	const int32 PrimitiveIndex = FMath::Trunc(FloatPrimitiveIndex * Primitives.Num());
+	const int32 PrimitiveIndex = FMath::TruncToInt(FloatPrimitiveIndex * Primitives.Num());
 	check(PrimitiveIndex >= 0 && PrimitiveIndex < Primitives.Num());
 
 	const FMeshLightPrimitive& SelectedPrimitive = Primitives[PrimitiveIndex];
diff --git a/Engine/Source/Programs/UnrealLightmass/Private/Lighting/FinalGather.cpp b/Engine/Source/Programs/UnrealLightmass/Private/Lighting/FinalGather.cpp
--- a/Engine/Source/Programs/UnrealLightmass/Private/Lighting/FinalGather.cpp
+++ b/Engine/Source/Programs/UnrealLightmass/Private/Lighting/FinalGather.cpp
@@ -308,7 +308,7 @@ FGatheredLightSample FStaticLightingSystem::IncomingRadianceImportancePhotons(
 				// Select a direction with uniform probability in a cone around the photon's incident direction
 				// See the "Extended Photon Map Implementation" paper
 				//@todo - select a direction with probability proportional to the power of the photon from that direction
-				const int32 PhotonIndex = FMath::Trunc(RandomStream.GetFraction() * ImportancePhotonDirections.Num());
+				const int32 PhotonIndex = FMath::TruncToInt(RandomStream.GetFraction() * ImportancePhotonDirections.Num());
 				checkSlow(PhotonIndex >= 0 && PhotonIndex < ImportancePhotonDirections.Num());
 				const FVector4& CurrentPhotonDirection = ImportancePhotonDirections[PhotonIndex];
 				checkSlow(Dot3(CurrentPhotonDirection, Vertex.TriangleNormal) > 0);
@@ -614,9 +614,9 @@ public:
 	const FLightingAndOcclusion& GetValue(FVector2D UV)
 	{
 		// Theta is radius, clamp
-		const int32 ThetaIndex = FMath::Clamp(FMath::Floor(UV.X * NumThetaSteps), 0, NumThetaSteps - 1);
+		const int32 ThetaIndex = FMath::Clamp(FMath::FloorToInt(UV.X * NumThetaSteps), 0, NumThetaSteps - 1);
 		// Phi is angle around the hemisphere axis, wrap on both ends
-		const int32 PhiIndex = (FMath::Floor(UV.Y * NumPhiSteps) + NumPhiSteps) % NumPhiSteps;
+		const int32 PhiIndex = (FMath::FloorToInt(UV.Y * NumPhiSteps) + NumPhiSteps) % NumPhiSteps;
 		const float CellU = FMath::Fractional(FMath::Clamp(UV.X, 0.0f, .9999f) * NumThetaSteps);
 		const float CellV = FMath::Abs(FMath::Fractional(UV.Y * NumPhiSteps));
 		const FSimpleQuadTree<FRefinementElement>& QuadTree = Cells[ThetaIndex * NumPhiSteps + PhiIndex];
@@ -1048,7 +1048,7 @@ SampleType FStaticLightingSystem::IncomingRadianceAdaptive(
 
 	FFinalGatherInfo FinalGatherInfo;
 
-	const int32 NumThetaSteps = FMath::Trunc(FMath::Sqrt(UniformHemisphereSamples.Num() / (float)PI) + .5f);
+	const int32 NumThetaSteps = FMath::TruncToInt(FMath::Sqrt(UniformHemisphereSamples.Num() / (float)PI) + .5f);
 	const int32 NumPhiSteps = UniformHemisphereSamples.Num() / NumThetaSteps;
 	checkSlow(NumThetaSteps * NumPhiSteps == UniformHemisphereSamples.Num());
 
@@ -1430,9 +1430,9 @@ void FStaticLightingSystem::CalculateIrradianceGradients(
 	{
 		// Extract Theta and Phi steps from the number of hemisphere samples requested
 		const float NumThetaStepsFloat = FMath::Sqrt(GetNumUniformHemisphereSamples(BounceNumber) / (float)PI);
-		const int32 NumThetaSteps = FMath::Trunc(NumThetaStepsFloat);
+		const int32 NumThetaSteps = FMath::TruncToInt(NumThetaStepsFloat);
 		// Using PI times more Phi steps as Theta steps
-		const int32 NumPhiSteps = FMath::Trunc(NumThetaStepsFloat * (float)PI);
+		const int32 NumPhiSteps = FMath::TruncToInt(NumThetaStepsFloat * (float)PI);
 		checkSlow(NumThetaSteps > 0 && NumPhiSteps > 0);
 
 		// Calculate the rotational gradient
diff --git a/Engine/Source/Programs/UnrealLightmass/Private/Lighting/LightingMesh.cpp b/Engine/Source/Programs/UnrealLightmass/Private/Lighting/LightingMesh.cpp
--- a/Engine/Source/Programs/UnrealLightmass/Private/Lighting/LightingMesh.cpp
+++ b/Engine/Source/Programs/UnrealLightmass/Private/Lighting/LightingMesh.cpp
@@ -233,12 +233,12 @@ void FStaticLightingMesh::CreateMeshAreaLights(
 	// if we have range -.2 to .3, we need space for the -1 .. 0 wrap, and the 0 to 1 wrap (ie 2 iterations)
 	// @todo UE4: Actually, if the Min was used to modify the UV when looping through the Corners array,
 	// we wouldn't need full integer ranges
-	const int32 NumIterationsX = (FMath::Floor(MaxUV.X) - FMath::Floor(MinUV.X)) + 1;
-	const int32 NumIterationsY = (FMath::Floor(MaxUV.Y) - FMath::Floor(MinUV.Y)) + 1;
+	const int32 NumIterationsX = (FMath::FloorToInt(MaxUV.X) - FMath::FloorToInt(MinUV.X)) + 1;
+	const int32 NumIterationsY = (FMath::FloorToInt(MaxUV.Y) - FMath::FloorToInt(MinUV.Y)) + 1;
 
 	// calculate the bias and scale needed to map the random UV range into 0 .. NumIterations when rasterizing
 	// into the TexelToCornersMap
-	const FVector2D UVBias(-FMath::Floor(MinUV.X), -FMath::Floor(MinUV.Y));
+	const FVector2D UVBias(-FMath::FloorToFloat(MinUV.X), -FMath::FloorToFloat(MinUV.Y));
 	const FVector2D UVScale(1.0f / NumIterationsX, 1.0f / NumIterationsY);
 
 	for (int32 MaterialIndex = 0; MaterialIndex < MaterialElements.Num(); MaterialIndex++)
diff --git a/Engine/Source/Programs/UnrealLightmass/Private/Lighting/LightingSystem.cpp b/Engine/Source/Programs/UnrealLightmass/Private/Lighting/LightingSystem.cpp
--- a/Engine/Source/Programs/UnrealLightmass/Private/Lighting/LightingSystem.cpp
+++ b/Engine/Source/Programs/UnrealLightmass/Private/Lighting/LightingSystem.cpp
@@ -767,7 +767,7 @@ int32 FStaticLightingSystem::GetNumUniformHemisphereSamples(int32 BounceNumber)
 int32 FStaticLightingSystem::GetNumPhotonImportanceHemisphereSamples() const
 {
 	return PhotonMappingSettings.bUsePhotonMapping ? 
-		FMath::Trunc(ImportanceTracingSettings.NumHemisphereSamples * PhotonMappingSettings.FinalGatherImportanceSampleFraction) : 0;
+		FMath::TruncToInt(ImportanceTracingSettings.NumHemisphereSamples * PhotonMappingSettings.FinalGatherImportanceSampleFraction) : 0;
 }
 
 FBoxSphereBounds FStaticLightingSystem::GetImportanceBounds(bool bClampToScene) const
@@ -833,7 +833,7 @@ void FStaticLightingSystem::ValidateSettings(FScene& InScene)
 		InScene.PhotonMappingSettings.bCacheIrradiancePhotonsOnSurfaces = false;
 	}
 	InScene.PhotonMappingSettings.FinalGatherImportanceSampleFraction = FMath::Clamp(InScene.PhotonMappingSettings.FinalGatherImportanceSampleFraction, 0.0f, 1.0f);
-	if (FMath::Trunc(InScene.ImportanceTracingSettings.NumHemisphereSamples * (1.0f - InScene.PhotonMappingSettings.FinalGatherImportanceSampleFraction) < 1))
+	if (FMath::TruncToInt(InScene.ImportanceTracingSettings.NumHemisphereSamples * (1.0f - InScene.PhotonMappingSettings.FinalGatherImportanceSampleFraction) < 1))
 	{
 		// Irradiance caching needs some uniform samples
 		InScene.IrradianceCachingSettings.bAllowIrradianceCaching = false;
@@ -1244,7 +1244,7 @@ void FStaticLightingSystem::UpdateInternalStatus(int32 OldNumTexelsCompleted) co
 	const float PreviousCompletedFraction = OldNumTexelsCompleted * InvTotal;
 	const float CurrentCompletedFraction = NumTexelsCompleted * InvTotal;
 	// Only log NumProgressSteps times
-	if (FMath::Trunc(PreviousCompletedFraction * NumProgressSteps) < FMath::Trunc(CurrentCompletedFraction * NumProgressSteps))
+	if (FMath::TruncToInt(PreviousCompletedFraction * NumProgressSteps) < FMath::TruncToInt(CurrentCompletedFraction * NumProgressSteps))
 	{
 		LogSolverMessage(FString::Printf(TEXT("Lighting %.1f%%"), CurrentCompletedFraction * 100.0f));
 	}
@@ -1264,7 +1264,7 @@ void FStaticLightingSystem::CacheSamples()
 			* (ImportanceTracingSettings.bUseAdaptiveSolver ? 1 : (1.0f - PhotonMappingSettings.FinalGatherImportanceSampleFraction))
 			* GeneralSettings.IndirectLightingQuality;
 
-		NumUniformHemisphereSamples = FMath::Trunc(NumSamplesFloat);
+		NumUniformHemisphereSamples = FMath::TruncToInt(NumSamplesFloat);
 	}
 	else
 	{
@@ -1279,8 +1279,8 @@ void FStaticLightingSystem::CacheSamples()
 		// Split the sampling domain up into cells with equal area
 		// Using PI times more Phi steps as Theta steps, but the relationship between them could be anything
 		const float NumThetaStepsFloat = FMath::Sqrt(NumUniformHemisphereSamples / (float)PI);
-		const int32 NumThetaSteps = FMath::Trunc(NumThetaStepsFloat);
-		const int32 NumPhiSteps = FMath::Trunc(NumThetaStepsFloat * (float)PI);
+		const int32 NumThetaSteps = FMath::TruncToInt(NumThetaStepsFloat);
+		const int32 NumPhiSteps = FMath::TruncToInt(NumThetaStepsFloat * (float)PI);
 
 		if (ImportanceTracingSettings.bUseCosinePDF)
 		{
@@ -1902,7 +1902,7 @@ void FStaticLightingSystem::CalculateApproximateDirectLighting(
 				FLinearColor UnnormalizedTransmission = FLinearColor::Black;
 
 				const TArray<FLightSurfaceSample>& LightSurfaceSamples = Light->GetCachedSurfaceSamples(0, false);
-				const int32 NumSamplesToTrace = FMath::Max(FMath::Trunc(LightSurfaceSamples.Num() * LightSampleFraction), 1);
+				const int32 NumSamplesToTrace = FMath::Max(FMath::TruncToInt(LightSurfaceSamples.Num() * LightSampleFraction), 1);
 
 				for (int32 RayIndex = 0; RayIndex < NumSamplesToTrace; RayIndex++)
 				{
@@ -2024,17 +2024,17 @@ void FStaticLightingSystem::BeginCalculateDominantShadowInfo(FGuid LightGuid)
 		DominantLightShadowInfo.LightSpaceImportanceBoundMin = LightSpaceImportanceBounds.Min;
 		DominantLightShadowInfo.LightSpaceImportanceBoundMax = LightSpaceImportanceBounds.Max;
 
-		DominantLightShadowInfo.ShadowMapSizeX = FMath::Trunc(FMath::Max(LightSpaceImportanceBounds.GetExtent().X * 2.0f / ShadowSettings.DominantShadowTransitionSampleDistanceX, 4.0f));
+		DominantLightShadowInfo.ShadowMapSizeX = FMath::TruncToInt(FMath::Max(LightSpaceImportanceBounds.GetExtent().X * 2.0f / ShadowSettings.DominantShadowTransitionSampleDistanceX, 4.0f));
 		DominantLightShadowInfo.ShadowMapSizeX = DominantLightShadowInfo.ShadowMapSizeX == appTruncErrorCode ? INT_MAX : DominantLightShadowInfo.ShadowMapSizeX;
-		DominantLightShadowInfo.ShadowMapSizeY = FMath::Trunc(FMath::Max(LightSpaceImportanceBounds.GetExtent().Y * 2.0f / ShadowSettings.DominantShadowTransitionSampleDistanceY, 4.0f));
+		DominantLightShadowInfo.ShadowMapSizeY = FMath::TruncToInt(FMath::Max(LightSpaceImportanceBounds.GetExtent().Y * 2.0f / ShadowSettings.DominantShadowTransitionSampleDistanceY, 4.0f));
 		DominantLightShadowInfo.ShadowMapSizeY = DominantLightShadowInfo.ShadowMapSizeY == appTruncErrorCode ? INT_MAX : DominantLightShadowInfo.ShadowMapSizeY;
 
 		// Clamp the number of dominant shadow samples generated if necessary while maintaining aspect ratio
 		if ((uint64)DominantLightShadowInfo.ShadowMapSizeX * (uint64)DominantLightShadowInfo.ShadowMapSizeY > (uint64)ShadowSettings.DominantShadowMaxSamples)
 		{
 			const float AspectRatio = DominantLightShadowInfo.ShadowMapSizeX / (float)DominantLightShadowInfo.ShadowMapSizeY;
-			DominantLightShadowInfo.ShadowMapSizeY = FMath::Trunc(FMath::Sqrt(ShadowSettings.DominantShadowMaxSamples / AspectRatio));
-			DominantLightShadowInfo.ShadowMapSizeX = FMath::Trunc(ShadowSettings.DominantShadowMaxSamples / DominantLightShadowInfo.ShadowMapSizeY);
+			DominantLightShadowInfo.ShadowMapSizeY = FMath::TruncToInt(FMath::Sqrt(ShadowSettings.DominantShadowMaxSamples / AspectRatio));
+			DominantLightShadowInfo.ShadowMapSizeX = FMath::TruncToInt(ShadowSettings.DominantShadowMaxSamples / DominantLightShadowInfo.ShadowMapSizeY);
 		}
 
 		// Allocate the shadow map
@@ -2061,7 +2061,7 @@ void FStaticLightingSystem::BeginCalculateDominantShadowInfo(FGuid LightGuid)
 		DominantLightShadowInfo.LightSpaceImportanceBoundMin = FVector4(-HalfCrossSectionLength, -HalfCrossSectionLength, 0);
 		DominantLightShadowInfo.LightSpaceImportanceBoundMax = FVector4(HalfCrossSectionLength, HalfCrossSectionLength, DominantSpotLight->Radius);
 
-		DominantLightShadowInfo.ShadowMapSizeX = FMath::Trunc(FMath::Max(HalfCrossSectionLength / ShadowSettings.DominantShadowTransitionSampleDistanceX, 4.0f));
+		DominantLightShadowInfo.ShadowMapSizeX = FMath::TruncToInt(FMath::Max(HalfCrossSectionLength / ShadowSettings.DominantShadowTransitionSampleDistanceX, 4.0f));
 		DominantLightShadowInfo.ShadowMapSizeX = DominantLightShadowInfo.ShadowMapSizeX == appTruncErrorCode ? INT_MAX : DominantLightShadowInfo.ShadowMapSizeX;
 		DominantLightShadowInfo.ShadowMapSizeY = DominantLightShadowInfo.ShadowMapSizeX;
 
@@ -2121,7 +2121,7 @@ void FStaticLightingSystem::BeginCalculateDominantShadowInfo(FGuid LightGuid)
 				}
 
 				// Quantize the distance into a uint16 to reduce memory usage in Unreal
-				const uint16 QuantizedDistance = FMath::Trunc(65535 * FMath::Clamp(MaxSampleDistance / MaxPossibleDistance, 0.0f, 1.0f));
+				const uint16 QuantizedDistance = FMath::TruncToInt(65535 * FMath::Clamp(MaxSampleDistance / MaxPossibleDistance, 0.0f, 1.0f));
 				DominantLightShadowInfo.ShadowMap[Y * DominantLightShadowInfo.ShadowMapSizeX + X] = FDominantLightShadowSample(QuantizedDistance);
 			}
 		}
@@ -2186,7 +2186,7 @@ void FStaticLightingSystem::CalculateDominantShadowInfoWorkRange(int32 ShadowMap
 		}
 
 		// Quantize the distance into a uint16 to reduce memory usage in Unreal
-		const uint16 QuantizedDistance = FMath::Trunc(65535 * FMath::Clamp(MaxSampleDistance * InvDistanceRange, 0.0f, 1.0f));
+		const uint16 QuantizedDistance = FMath::TruncToInt(65535 * FMath::Clamp(MaxSampleDistance * InvDistanceRange, 0.0f, 1.0f));
 		DominantLightShadowInfo.ShadowMap[ShadowMapY * DominantLightShadowInfo.ShadowMapSizeX + X] = FDominantLightShadowSample(QuantizedDistance);
 	}
 	Context.Stats.DominantShadowThreadTime = FPlatformTime::Seconds() - CalculateWorkRangeStart;
diff --git a/Engine/Source/Programs/UnrealLightmass/Private/Lighting/LightmapData.cpp b/Engine/Source/Programs/UnrealLightmass/Private/Lighting/LightmapData.cpp
--- a/Engine/Source/Programs/UnrealLightmass/Private/Lighting/LightmapData.cpp
+++ b/Engine/Source/Programs/UnrealLightmass/Private/Lighting/LightmapData.cpp
@@ -274,11 +274,11 @@ namespace Lightmass
 			const float BentNormalLength = BentNormal.Size();
 			const FVector NormalizedBentNormal = BentNormal.SafeNormal() * FVector(.5f) + FVector(.5f);
 
-			DestCoefficients.SkyOcclusion[0] = (uint8)FMath::Clamp<int32>( FMath::Round( NormalizedBentNormal[0] * 255.0f ), 0, 255 );
-			DestCoefficients.SkyOcclusion[1] = (uint8)FMath::Clamp<int32>( FMath::Round( NormalizedBentNormal[1] * 255.0f ), 0, 255 );
-			DestCoefficients.SkyOcclusion[2] = (uint8)FMath::Clamp<int32>( FMath::Round( NormalizedBentNormal[2] * 255.0f ), 0, 255 );
+			DestCoefficients.SkyOcclusion[0] = (uint8)FMath::Clamp<int32>( FMath::RoundToInt( NormalizedBentNormal[0] * 255.0f ), 0, 255 );
+			DestCoefficients.SkyOcclusion[1] = (uint8)FMath::Clamp<int32>( FMath::RoundToInt( NormalizedBentNormal[1] * 255.0f ), 0, 255 );
+			DestCoefficients.SkyOcclusion[2] = (uint8)FMath::Clamp<int32>( FMath::RoundToInt( NormalizedBentNormal[2] * 255.0f ), 0, 255 );
 			// Sqrt on length to allocate more precision near 0
-			DestCoefficients.SkyOcclusion[3] = (uint8)FMath::Clamp<int32>( FMath::Round( FMath::Sqrt(BentNormalLength) * 255.0f ), 0, 255 );
+			DestCoefficients.SkyOcclusion[3] = (uint8)FMath::Clamp<int32>( FMath::RoundToInt( FMath::Sqrt(BentNormalLength) * 255.0f ), 0, 255 );
 
 			{
 				float L, U, V, W;
@@ -292,24 +292,24 @@ namespace Lightmass
 				W = W * CoefficientMultiply[0][2] + CoefficientAdd[0][2];
 				LogL = LogL * CoefficientMultiply[0][3] + CoefficientAdd[0][3];
 
-				float Residual = LogL * 255.0f - FMath::Round( LogL * 255.0f ) + 0.5f;
+				float Residual = LogL * 255.0f - FMath::RoundToFloat( LogL * 255.0f ) + 0.5f;
 
 				// U, V, W, LogL
 				// UVW stored in gamma space
-				DestCoefficients.Coefficients[0][0] = (uint8)FMath::Clamp<int32>( FMath::Round( FMath::Pow( U, 1.0f / 2.2f ) * 255.0f ), 0, 255 );
-				DestCoefficients.Coefficients[0][1] = (uint8)FMath::Clamp<int32>( FMath::Round( FMath::Pow( V, 1.0f / 2.2f ) * 255.0f ), 0, 255 );
-				DestCoefficients.Coefficients[0][2] = (uint8)FMath::Clamp<int32>( FMath::Round( FMath::Pow( W, 1.0f / 2.2f ) * 255.0f ), 0, 255 );
-				DestCoefficients.Coefficients[0][3] = (uint8)FMath::Clamp<int32>( FMath::Round( LogL * 255.0f ), 0, 255 );
+				DestCoefficients.Coefficients[0][0] = (uint8)FMath::Clamp<int32>( FMath::RoundToInt( FMath::Pow( U, 1.0f / 2.2f ) * 255.0f ), 0, 255 );
+				DestCoefficients.Coefficients[0][1] = (uint8)FMath::Clamp<int32>( FMath::RoundToInt( FMath::Pow( V, 1.0f / 2.2f ) * 255.0f ), 0, 255 );
+				DestCoefficients.Coefficients[0][2] = (uint8)FMath::Clamp<int32>( FMath::RoundToInt( FMath::Pow( W, 1.0f / 2.2f ) * 255.0f ), 0, 255 );
+				DestCoefficients.Coefficients[0][3] = (uint8)FMath::Clamp<int32>( FMath::RoundToInt( LogL * 255.0f ), 0, 255 );
 				
 				float Dx = SourceSample.Coefficients[1][0] * CoefficientMultiply[1][0] + CoefficientAdd[1][0];
 				float Dy = SourceSample.Coefficients[1][1] * CoefficientMultiply[1][1] + CoefficientAdd[1][1];
 				float Dz = SourceSample.Coefficients[1][2] * CoefficientMultiply[1][2] + CoefficientAdd[1][2];
 
 				// Dx, Dy, Dz, Residual
-				DestCoefficients.Coefficients[1][0] = (uint8)FMath::Clamp<int32>( FMath::Round( Dx * 255.0f ), 0, 255 );
-				DestCoefficients.Coefficients[1][1] = (uint8)FMath::Clamp<int32>( FMath::Round( Dy * 255.0f ), 0, 255 );
-				DestCoefficients.Coefficients[1][2] = (uint8)FMath::Clamp<int32>( FMath::Round( Dz * 255.0f ), 0, 255 );
-				DestCoefficients.Coefficients[1][3] = (uint8)FMath::Clamp<int32>( FMath::Round( Residual * 255.0f ), 0, 255 );
+				DestCoefficients.Coefficients[1][0] = (uint8)FMath::Clamp<int32>( FMath::RoundToInt( Dx * 255.0f ), 0, 255 );
+				DestCoefficients.Coefficients[1][1] = (uint8)FMath::Clamp<int32>( FMath::RoundToInt( Dy * 255.0f ), 0, 255 );
+				DestCoefficients.Coefficients[1][2] = (uint8)FMath::Clamp<int32>( FMath::RoundToInt( Dz * 255.0f ), 0, 255 );
+				DestCoefficients.Coefficients[1][3] = (uint8)FMath::Clamp<int32>( FMath::RoundToInt( Residual * 255.0f ), 0, 255 );
 			}
 
 			{
@@ -324,9 +324,9 @@ namespace Lightmass
 				float LogB = LogL * W * CoefficientMultiply[2][2] + CoefficientAdd[2][2];
 				
 				// LogR, LogG, LogB
-				DestCoefficients.Coefficients[2][0] = (uint8)FMath::Clamp<int32>( FMath::Round( LogR * 255.0f ), 0, 255 );
-				DestCoefficients.Coefficients[2][1] = (uint8)FMath::Clamp<int32>( FMath::Round( LogG * 255.0f ), 0, 255 );
-				DestCoefficients.Coefficients[2][2] = (uint8)FMath::Clamp<int32>( FMath::Round( LogB * 255.0f ), 0, 255 );
+				DestCoefficients.Coefficients[2][0] = (uint8)FMath::Clamp<int32>( FMath::RoundToInt( LogR * 255.0f ), 0, 255 );
+				DestCoefficients.Coefficients[2][1] = (uint8)FMath::Clamp<int32>( FMath::RoundToInt( LogG * 255.0f ), 0, 255 );
+				DestCoefficients.Coefficients[2][2] = (uint8)FMath::Clamp<int32>( FMath::RoundToInt( LogB * 255.0f ), 0, 255 );
 				DestCoefficients.Coefficients[2][3] = 255;
 
 				float Dx = SourceSample.Coefficients[3][0] * CoefficientMultiply[3][0] + CoefficientAdd[3][0];
@@ -334,9 +334,9 @@ namespace Lightmass
 				float Dz = SourceSample.Coefficients[3][2] * CoefficientMultiply[3][2] + CoefficientAdd[3][2];
 
 				// Dx, Dy, Dz
-				DestCoefficients.Coefficients[3][0] = (uint8)FMath::Clamp<int32>( FMath::Round( Dx * 255.0f ), 0, 255 );
-				DestCoefficients.Coefficients[3][1] = (uint8)FMath::Clamp<int32>( FMath::Round( Dy * 255.0f ), 0, 255 );
-				DestCoefficients.Coefficients[3][2] = (uint8)FMath::Clamp<int32>( FMath::Round( Dz * 255.0f ), 0, 255 );
+				DestCoefficients.Coefficients[3][0] = (uint8)FMath::Clamp<int32>( FMath::RoundToInt( Dx * 255.0f ), 0, 255 );
+				DestCoefficients.Coefficients[3][1] = (uint8)FMath::Clamp<int32>( FMath::RoundToInt( Dy * 255.0f ), 0, 255 );
+				DestCoefficients.Coefficients[3][2] = (uint8)FMath::Clamp<int32>( FMath::RoundToInt( Dz * 255.0f ), 0, 255 );
 				DestCoefficients.Coefficients[3][3] = 255;
 			}
 		}
@@ -370,7 +370,7 @@ namespace Lightmass
 			}
 #endif
 			// Convert linear input values to gamma space before quantizing, which preserves more detail in the darks where banding would be noticeable otherwise
-			QuantizedValue.Visibility = (uint8)FMath::Clamp<int32>(FMath::Trunc(FMath::Pow(Value.Visibility, 1.0f / 2.2f) * 255.0f), 0, 255);
+			QuantizedValue.Visibility = (uint8)FMath::Clamp<int32>(FMath::TruncToInt(FMath::Pow(Value.Visibility, 1.0f / 2.2f) * 255.0f), 0, 255);
 			QuantizedValue.Coverage = Value.bIsMapped ? 255 : 0;
 		}
 		Data.Empty();
@@ -408,8 +408,8 @@ namespace Lightmass
 				int32 TempBreak = 0;
 			}
 #endif
-			QuantizedValue.Distance = (uint8)FMath::Clamp<int32>(FMath::Trunc(Value.Distance * 255.0f), 0, 255);
-			QuantizedValue.PenumbraSize = (uint8)FMath::Clamp<int32>(FMath::Trunc(Value.PenumbraSize * 255.0f), 0, 255);
+			QuantizedValue.Distance = (uint8)FMath::Clamp<int32>(FMath::TruncToInt(Value.Distance * 255.0f), 0, 255);
+			QuantizedValue.PenumbraSize = (uint8)FMath::Clamp<int32>(FMath::TruncToInt(Value.PenumbraSize * 255.0f), 0, 255);
 			QuantizedValue.Coverage = Value.bIsMapped ? 255 : 0;
 		}
 		Data.Empty();
diff --git a/Engine/Source/Programs/UnrealLightmass/Private/Lighting/PhotonMapping.cpp b/Engine/Source/Programs/UnrealLightmass/Private/Lighting/PhotonMapping.cpp
--- a/Engine/Source/Programs/UnrealLightmass/Private/Lighting/PhotonMapping.cpp
+++ b/Engine/Source/Programs/UnrealLightmass/Private/Lighting/PhotonMapping.cpp
@@ -76,16 +76,16 @@ void FStaticLightingSystem::InitializePhotonSettings()
 	// If the importance volume is valid, only gather enough indirect photon paths to meet IndirectPhotonPathDensity inside the importance volume
 	if (!PhotonMappingSettings.bEmitPhotonsOutsideImportanceVolume && ImportanceBounds.SphereRadius > DELTA)
 	{ 
-		NumIndirectPhotonPaths = FMath::Trunc(Scene.PhotonMappingSettings.IndirectPhotonPathDensity * ImportanceSurfaceAreaMillionUnits);
+		NumIndirectPhotonPaths = FMath::TruncToInt(Scene.PhotonMappingSettings.IndirectPhotonPathDensity * ImportanceSurfaceAreaMillionUnits);
 	}
 	else if (ImportanceBounds.SphereRadius > DELTA)
 	{
-		NumIndirectPhotonPaths = FMath::Trunc(Scene.PhotonMappingSettings.IndirectPhotonPathDensity * ImportanceSurfaceAreaMillionUnits
+		NumIndirectPhotonPaths = FMath::TruncToInt(Scene.PhotonMappingSettings.IndirectPhotonPathDensity * ImportanceSurfaceAreaMillionUnits
 			+ Scene.PhotonMappingSettings.OutsideImportanceVolumeDensityScale * Scene.PhotonMappingSettings.IndirectPhotonPathDensity * SceneSurfaceAreaMillionUnits);
 	}
 	else
 	{
-		NumIndirectPhotonPaths = FMath::Trunc(Scene.PhotonMappingSettings.IndirectPhotonPathDensity * SceneSurfaceAreaMillionUnits);
+		NumIndirectPhotonPaths = FMath::TruncToInt(Scene.PhotonMappingSettings.IndirectPhotonPathDensity * SceneSurfaceAreaMillionUnits);
 	}
 	NumIndirectPhotonPaths = NumIndirectPhotonPaths == appTruncErrorCode ? MaxNumIndirectPhotonPaths : NumIndirectPhotonPaths;
 	NumIndirectPhotonPaths = FMath::Min(NumIndirectPhotonPaths, MaxNumIndirectPhotonPaths);
@@ -474,7 +474,7 @@ void FStaticLightingSystem::EmitDirectPhotonsWorkRange(
 	{
 		// Attempt to preallocate irradiance photons based on the percentage of photons that go into the irradiance photon map.
 		// The actual number of irradiance photons is based on probability.
-		Output.IrradiancePhotons->Empty(FMath::Trunc(DirectIrradiancePhotonFraction * DirectPhotonEfficiency * WorkRange.NumDirectPhotonsToEmit));
+		Output.IrradiancePhotons->Empty(FMath::TruncToInt(DirectIrradiancePhotonFraction * DirectPhotonEfficiency * WorkRange.NumDirectPhotonsToEmit));
 	}
 
 	FCoherentRayCache CoherentRayCache;
@@ -485,7 +485,7 @@ void FStaticLightingSystem::EmitDirectPhotonsWorkRange(
 
 	// Array of rays from each light which resulted in an indirect path.
 	// These are used in the second emitting pass to guide light sampling for indirect photons.
-	Output.DirectPhotons.Empty(FMath::Trunc(WorkRange.NumDirectPhotonsToEmit * DirectPhotonEfficiency));
+	Output.DirectPhotons.Empty(FMath::TruncToInt(WorkRange.NumDirectPhotonsToEmit * DirectPhotonEfficiency));
 	
 	Output.NumPhotonsEmitted = 0;
 	int32 NumIndirectPathRaysGathered = 0;
@@ -512,7 +512,7 @@ void FStaticLightingSystem::EmitDirectPhotonsWorkRange(
 		float LightIndex;
 		// Pick a light with probability proportional to the light's fraction of the direct photons being gathered for the whole scene
 		Sample1dCDF(Input.LightDistribution.LightPDFs, Input.LightDistribution.LightCDFs, Input.LightDistribution.UnnormalizedIntegral, RandomStream, LightPDF, LightIndex);
-		const int32 QuantizedLightIndex = FMath::Trunc(LightIndex * Input.LightDistribution.LightPDFs.Num());
+		const int32 QuantizedLightIndex = FMath::TruncToInt(LightIndex * Input.LightDistribution.LightPDFs.Num());
 		check(QuantizedLightIndex >= 0 && QuantizedLightIndex < Lights.Num());
 		const FLight* Light = Lights[QuantizedLightIndex];
 
@@ -860,13 +860,13 @@ void FStaticLightingSystem::EmitIndirectPhotonsWorkRange(
 	}
 
 	//@todo - re-evaluate these sizes
-	Output.FirstBouncePhotons.Empty(FMath::Trunc(WorkRange.NumIndirectPhotonsToEmit * .6f * IndirectPhotonEfficiency));
-	Output.SecondBouncePhotons.Empty(FMath::Trunc(WorkRange.NumIndirectPhotonsToEmit * .4f * IndirectPhotonEfficiency));
+	Output.FirstBouncePhotons.Empty(FMath::TruncToInt(WorkRange.NumIndirectPhotonsToEmit * .6f * IndirectPhotonEfficiency));
+	Output.SecondBouncePhotons.Empty(FMath::TruncToInt(WorkRange.NumIndirectPhotonsToEmit * .4f * IndirectPhotonEfficiency));
 	if (PhotonMappingSettings.bUseIrradiancePhotons)
 	{
 		// Attempt to preallocate irradiance photons based on the percentage of photons that go into the irradiance photon map.
 		// The actual number of irradiance photons is based on probability.
-		Output.IrradiancePhotons->Empty(FMath::Trunc(IndirectIrradiancePhotonFraction * IndirectPhotonEfficiency * WorkRange.NumIndirectPhotonsToEmit));
+		Output.IrradiancePhotons->Empty(FMath::TruncToInt(IndirectIrradiancePhotonFraction * IndirectPhotonEfficiency * WorkRange.NumIndirectPhotonsToEmit));
 	}
 
 	FCoherentRayCache CoherentRayCache;
@@ -893,7 +893,7 @@ void FStaticLightingSystem::EmitIndirectPhotonsWorkRange(
 			float LightIndex;
 			// Pick a light with probability proportional to the light's fraction of the scene's light power
 			Sample1dCDF(Input.LightDistribution.LightPDFs, Input.LightDistribution.LightCDFs, Input.LightDistribution.UnnormalizedIntegral, RandomStream, LightPDF, LightIndex);
-			const int32 QuantizedLightIndex = FMath::Trunc(LightIndex * Input.LightDistribution.LightPDFs.Num());
+			const int32 QuantizedLightIndex = FMath::TruncToInt(LightIndex * Input.LightDistribution.LightPDFs.Num());
 			check(QuantizedLightIndex >= 0 && QuantizedLightIndex < Lights.Num());
 			Light = Lights[QuantizedLightIndex];
 
@@ -1007,7 +1007,7 @@ void FStaticLightingSystem::EmitIndirectPhotonsWorkRange(
 						FScopeLock DebugOutputLock(&DebugOutputSync);
 						if (DebugOutput.IndirectPhotons.Num() == 0)
 						{
-							DebugOutput.IndirectPhotons.Empty(FMath::Trunc(NumIndirectPhotonsToEmit * IndirectPhotonEfficiency));
+							DebugOutput.IndirectPhotons.Empty(FMath::TruncToInt(NumIndirectPhotonsToEmit * IndirectPhotonEfficiency));
 						}
 						DebugOutput.IndirectPhotons.Add(FDebugPhoton(NewPhoton.GetId(), NewPhoton.GetPosition(), SampleRay.Start - NewPhoton.GetPosition(), NewPhoton.GetSurfaceNormal()));
 					}
