#ttp 331438 UE4: PlatformMath refactor

diff --git a/Engine/Source/Programs/UnrealLightmass/Private/Lighting/PrecomputedVisibility.cpp b/Engine/Source/Programs/UnrealLightmass/Private/Lighting/PrecomputedVisibility.cpp
--- a/Engine/Source/Programs/UnrealLightmass/Private/Lighting/PrecomputedVisibility.cpp
+++ b/Engine/Source/Programs/UnrealLightmass/Private/Lighting/PrecomputedVisibility.cpp
@@ -155,9 +155,9 @@ void FCellPlacementRasterPolicy::ProcessPixel(int32 X,int32 Y,const InterpolantT
 int32 FStaticLightingSystem::GetGroupCellIndex(FVector BoxCenter) const
 {
 	const FVector GridPosition = FVector(GroupVisibilityGridSizeXY, GroupVisibilityGridSizeXY, GroupVisibilityGridSizeZ) * (BoxCenter - VisibilityGridBounds.Min) / (VisibilityGridBounds.Max - VisibilityGridBounds.Min);
-	const int32 GridX = FMath::Trunc(GridPosition.X);
-	const int32 GridY = FMath::Trunc(GridPosition.Y);
-	const int32 GridZ = FMath::Trunc(GridPosition.Z);
+	const int32 GridX = FMath::TruncToInt(GridPosition.X);
+	const int32 GridY = FMath::TruncToInt(GridPosition.Y);
+	const int32 GridZ = FMath::TruncToInt(GridPosition.Z);
 	const int32 CellIndex = GridZ * GroupVisibilityGridSizeXY * GroupVisibilityGridSizeZ + GridY * GroupVisibilityGridSizeXY + GridX;
 	
 	if (GridX > 0 && GridX < GroupVisibilityGridSizeXY 
@@ -192,8 +192,8 @@ void FStaticLightingSystem::SetupPrecomputedVisibility()
 
 	const FBoxSphereBounds PrecomputedVisibilityBounds = Scene.GetVisibilityVolumeBounds();
 	const FVector4 VolumeSizes = PrecomputedVisibilityBounds.BoxExtent * 2.0f / PrecomputedVisibilitySettings.CellSize;
-	const int32 SizeX = FMath::Trunc(VolumeSizes.X + DELTA) + 1;
-	const int32 SizeY = FMath::Trunc(VolumeSizes.Y + DELTA) + 1;
+	const int32 SizeX = FMath::TruncToInt(VolumeSizes.X + DELTA) + 1;
+	const int32 SizeY = FMath::TruncToInt(VolumeSizes.Y + DELTA) + 1;
 
 	if (!PrecomputedVisibilitySettings.bPlaceCellsOnlyAlongCameraTracks)
 	{
@@ -532,7 +532,7 @@ void FStaticLightingSystem::SetupPrecomputedVisibility()
 		VisibilityGridBounds = FBox(0);
 		// Drop last 10% of meshes which will expand the grid bounds 
 		// This is to handle distant skybox type meshes
-		const int32 MaxMeshIndex = FMath::Min(FMath::Max(FMath::Trunc(.9f * SortMeshes.Num()), 1), SortMeshes.Num() - 1);
+		const int32 MaxMeshIndex = FMath::Min(FMath::Max(FMath::TruncToInt(.9f * SortMeshes.Num()), 1), SortMeshes.Num() - 1);
 
 		for (int32 MeshIndex = 0; MeshIndex < MaxMeshIndex; MeshIndex++)
 		{
@@ -543,7 +543,7 @@ void FStaticLightingSystem::SetupPrecomputedVisibility()
 		const float TargetNumGroupsAsFractionOfMeshes = .3f;
 		const int32 ZDimensionDivisor = 4;
 		// Determine grid X and Y size using SizeX * SizeY * SizeZ = NumMeshes * TargetNumGroupsAsFractionOfMeshes
-		GroupVisibilityGridSizeXY = FMath::Max(FMath::Trunc(FMath::Pow((ZDimensionDivisor * TargetNumGroupsAsFractionOfMeshes * VisibilityMeshes.Num()), 1.0f / 3.0f)), 1);
+		GroupVisibilityGridSizeXY = FMath::Max(FMath::TruncToInt(FMath::Pow((ZDimensionDivisor * TargetNumGroupsAsFractionOfMeshes * VisibilityMeshes.Num()), 1.0f / 3.0f)), 1);
 		GroupVisibilityGridSizeZ = FMath::Max(GroupVisibilityGridSizeXY / ZDimensionDivisor, 1);
 
 		const int32 GridSizeXY = GroupVisibilityGridSizeXY;
@@ -769,7 +769,7 @@ bool ComputeBoxVisibility(
 	const float SizeRatio = MeshSize / Distance;
 	// Use MaxMeshSamples for meshes with a large projected angle, and MinMeshSamples for meshes with a small projected angle
 	// Meshes with a large projected angle require more samples to determine visibility accurately
-	const int32 NumMeshSamples = FMath::Clamp(FMath::Trunc(SizeRatio * PrecomputedVisibilitySettings.MaxMeshSamples), PrecomputedVisibilitySettings.MinMeshSamples, PrecomputedVisibilitySettings.MaxMeshSamples);
+	const int32 NumMeshSamples = FMath::Clamp(FMath::TruncToInt(SizeRatio * PrecomputedVisibilitySettings.MaxMeshSamples), PrecomputedVisibilitySettings.MinMeshSamples, PrecomputedVisibilitySettings.MaxMeshSamples);
 
 	// Treat meshes whose projected angle is greater than 90 degrees as visible, since it becomes overly costly to determine if these are visible
 	// (consider a large close mesh that only has a tiny part visible)
@@ -797,13 +797,13 @@ bool ComputeBoxVisibility(
 					// Generate a sample on the visible faces of the cell, picking a face with probability proportional to the projected angle of the cell face onto the mesh's origin
 					//@todo - weight by face area, since cells have a different height from their x and y sizes
 					Sample1dCDF(VisibleCellFacePDFs, VisibleCellFaceCDFs, UnnormalizedIntegral, RandomStream, PDF, Sample);
-					const int32 ChosenCellFaceIndex = FMath::Trunc(Sample * VisibleCellFaces.Num());
+					const int32 ChosenCellFaceIndex = FMath::TruncToInt(Sample * VisibleCellFaces.Num());
 					const FAxisAlignedCellFace& ChosenFace = CellFaces[VisibleCellFaces[ChosenCellFaceIndex]];
 					NewSample.CellPosition = ChosenFace.FaceMin + ChosenFace.FaceExtent * FVector4(RandomStream.GetFraction(), RandomStream.GetFraction(), RandomStream.GetFraction());
 				}
 				{
 					// Generate a sample on the visible faces of the mesh
-					const int32 ChosenFaceIndex = FMath::Trunc(RandomStream.GetFraction() * VisibleMeshFaces.Num());
+					const int32 ChosenFaceIndex = FMath::TruncToInt(RandomStream.GetFraction() * VisibleMeshFaces.Num());
 					const FAxisAlignedCellFace& ChosenFace = MeshBoxFaces[VisibleMeshFaces[ChosenFaceIndex]];
 					NewSample.MeshPosition = ChosenFace.FaceMin + ChosenFace.FaceExtent * FVector4(RandomStream.GetFraction(), RandomStream.GetFraction(), RandomStream.GetFraction());
 				}
@@ -911,7 +911,7 @@ bool ComputeBoxVisibility(
 			for (int32 ImportanceSampleIndex = 0; ImportanceSampleIndex < PrecomputedVisibilitySettings.NumImportanceSamples && !bVisible; ImportanceSampleIndex++)
 			{
 				// Pick one of the furthest samples with uniform probability
-				const int32 SampleIndex = FMath::Trunc(RandomStream.GetFraction() * FurthestSamples.Num());
+				const int32 SampleIndex = FMath::TruncToInt(RandomStream.GetFraction() * FurthestSamples.Num());
 				const FVisibilityQuerySample& CurrentSample = *FurthestSamples[SampleIndex];
 				const float VectorLength = (CurrentSample.CellPosition - CurrentSample.MeshPosition).Size3();
 				const FVector4 CurrentDirection = (CurrentSample.MeshPosition - CurrentSample.CellPosition).SafeNormal();
diff --git a/Engine/Source/Programs/UnrealLightmass/Private/Lighting/Raster.h b/Engine/Source/Programs/UnrealLightmass/Private/Lighting/Raster.h
--- a/Engine/Source/Programs/UnrealLightmass/Private/Lighting/Raster.h
+++ b/Engine/Source/Programs/UnrealLightmass/Private/Lighting/Raster.h
@@ -111,8 +111,8 @@ private:
 		bool BackFacing
 		)
 	{
-		int32	IntMinY = FMath::Clamp(FMath::Ceil(MinY),RasterPolicyType::GetMinY(),RasterPolicyType::GetMaxY() + 1),
-			IntMaxY = FMath::Clamp(FMath::Ceil(MaxY),RasterPolicyType::GetMinY(),RasterPolicyType::GetMaxY() + 1);
+		int32	IntMinY = FMath::Clamp(FMath::CeilToInt(MinY),RasterPolicyType::GetMinY(),RasterPolicyType::GetMaxY() + 1),
+			IntMaxY = FMath::Clamp(FMath::CeilToInt(MaxY),RasterPolicyType::GetMinY(),RasterPolicyType::GetMaxY() + 1);
 
 		for(int32 IntY = IntMinY;IntY < IntMaxY;IntY++)
 		{
@@ -130,8 +130,8 @@ private:
 
 			if(MaxX > MinX)
 			{
-				int32				IntMinX = FMath::Clamp(FMath::Ceil(MinX),RasterPolicyType::GetMinX(),RasterPolicyType::GetMaxX() + 1),
-								IntMaxX = FMath::Clamp(FMath::Ceil(MaxX),RasterPolicyType::GetMinX(),RasterPolicyType::GetMaxX() + 1);
+				int32				IntMinX = FMath::Clamp(FMath::CeilToInt(MinX),RasterPolicyType::GetMinX(),RasterPolicyType::GetMaxX() + 1),
+								IntMaxX = FMath::Clamp(FMath::CeilToInt(MaxX),RasterPolicyType::GetMinX(),RasterPolicyType::GetMaxX() + 1);
 				InterpolantType	DeltaInterpolant = (MaxInterpolant - MinInterpolant) / (MaxX - MinX);
 
 				for(int32 X = IntMinX;X < IntMaxX;X++)
diff --git a/Engine/Source/Programs/UnrealLightmass/Private/Lighting/SampleVolume.cpp b/Engine/Source/Programs/UnrealLightmass/Private/Lighting/SampleVolume.cpp
--- a/Engine/Source/Programs/UnrealLightmass/Private/Lighting/SampleVolume.cpp
+++ b/Engine/Source/Programs/UnrealLightmass/Private/Lighting/SampleVolume.cpp
@@ -221,8 +221,8 @@ void FStaticLightingSystem::CalculateVolumeSamples()
 			// Octree used for interpolating lighting for debugging
 			VolumeLightingInterpolationOctree = FVolumeLightingInterpolationOctree(VolumeBounds.Origin, VolumeBounds.BoxExtent.GetMax());
 			// Determine the resolution that the scene should be rasterized at based on SurfaceLightSampleSpacing and the scene's extent
-			const int32 RasterSizeX = FMath::Trunc(2.0f * VolumeBounds.BoxExtent.X / DynamicObjectSettings.SurfaceLightSampleSpacing);
-			const int32 RasterSizeY = FMath::Trunc(2.0f * VolumeBounds.BoxExtent.Y / DynamicObjectSettings.SurfaceLightSampleSpacing);
+			const int32 RasterSizeX = FMath::TruncToInt(2.0f * VolumeBounds.BoxExtent.X / DynamicObjectSettings.SurfaceLightSampleSpacing);
+			const int32 RasterSizeY = FMath::TruncToInt(2.0f * VolumeBounds.BoxExtent.Y / DynamicObjectSettings.SurfaceLightSampleSpacing);
 
 			FTriangleRasterizer<FVolumeSamplePlacementRasterPolicy> Rasterizer(
 				FVolumeSamplePlacementRasterPolicy(
@@ -379,7 +379,7 @@ void FStaticLightingSystem::CalculateVolumeSamples()
 			}
 
 			const float VolumeSpacingCubed = DynamicObjectSettings.VolumeLightSampleSpacing * DynamicObjectSettings.VolumeLightSampleSpacing * DynamicObjectSettings.VolumeLightSampleSpacing;
-			int32 RequestedVolumeSamples = FMath::Trunc(8.0f * VolumeBounds.BoxExtent.X * VolumeBounds.BoxExtent.Y * VolumeBounds.BoxExtent.Z / VolumeSpacingCubed);
+			int32 RequestedVolumeSamples = FMath::TruncToInt(8.0f * VolumeBounds.BoxExtent.X * VolumeBounds.BoxExtent.Y * VolumeBounds.BoxExtent.Z / VolumeSpacingCubed);
 			RequestedVolumeSamples = RequestedVolumeSamples == appTruncErrorCode ? INT_MAX : RequestedVolumeSamples;
 			float EffectiveVolumeSpacing = DynamicObjectSettings.VolumeLightSampleSpacing;
 
@@ -424,8 +424,8 @@ void FStaticLightingSystem::CalculateVolumeSamples()
 			const int32 NumUpperVolumeSamples = ImportanceTracingSettings.NumHemisphereSamples * DynamicObjectSettings.NumHemisphereSamplesScale;
 			// Volume samples don't do any importance sampling so they need more samples for the same amount of variance as surface samples
 			const float NumThetaStepsFloat = FMath::Sqrt(NumUpperVolumeSamples / (float)PI);
-			const int32 NumThetaSteps = FMath::Trunc(NumThetaStepsFloat);
-			const int32 NumPhiSteps = FMath::Trunc(NumThetaStepsFloat * (float)PI);
+			const int32 NumThetaSteps = FMath::TruncToInt(NumThetaStepsFloat);
+			const int32 NumPhiSteps = FMath::TruncToInt(NumThetaStepsFloat * (float)PI);
 
 			GenerateStratifiedUniformHemisphereSamples(NumThetaSteps, NumPhiSteps, RandomStream, UniformHemisphereSamples, UniformHemisphereSampleUniforms);
 
diff --git a/Engine/Source/Programs/UnrealLightmass/Private/Lighting/Texture.h b/Engine/Source/Programs/UnrealLightmass/Private/Lighting/Texture.h
--- a/Engine/Source/Programs/UnrealLightmass/Private/Lighting/Texture.h
+++ b/Engine/Source/Programs/UnrealLightmass/Private/Lighting/Texture.h
@@ -82,7 +82,7 @@ namespace Lightmass
 
 		inline uint8* SampleRawPtr(const FVector2D& UV) const
 		{
-			// Wrapped addressing (uses FMath::Floor and not appFractional, as appFractional causes the
+			// Wrapped addressing (uses FMath::FloorToInt and not appFractional, as appFractional causes the
 			// following mapping:
 			// .4 -> .4
 			// -1.4 -> .4
@@ -94,8 +94,8 @@ namespace Lightmass
 			// (.4 - 1 = -.6) -> .4
 			//
 			// because when you subtract 1 from a UV it needs to have the exact same fractional part
- 			const int32 X = FMath::Clamp(FMath::Trunc((UV.X - FMath::Floor(UV.X)) * SizeX), 0, SizeX - 1);
- 			const int32 Y = FMath::Clamp(FMath::Trunc((UV.Y - FMath::Floor(UV.Y)) * SizeY), 0, SizeY - 1);
+ 			const int32 X = FMath::Clamp(FMath::TruncToInt((UV.X - FMath::FloorToInt(UV.X)) * SizeX), 0, SizeX - 1);
+ 			const int32 Y = FMath::Clamp(FMath::TruncToInt((UV.Y - FMath::FloorToInt(UV.Y)) * SizeY), 0, SizeY - 1);
 			// Byte index into Data
 			const int32 DataIndex = Y * SizeX * ElementSize + X * ElementSize;
 			return &Data[DataIndex];
diff --git a/Engine/Source/Programs/UnrealLightmass/Private/Lighting/TextureMapping.cpp b/Engine/Source/Programs/UnrealLightmass/Private/Lighting/TextureMapping.cpp
--- a/Engine/Source/Programs/UnrealLightmass/Private/Lighting/TextureMapping.cpp
+++ b/Engine/Source/Programs/UnrealLightmass/Private/Lighting/TextureMapping.cpp
@@ -269,8 +269,8 @@ void FStaticLightingSystem::CacheIrradiancePhotonsTextureMapping(FStaticLighting
 	const FBoxSphereBounds ImportanceBounds = Scene.GetImportanceBounds();
 
 	// Cache irradiance photons at a lower resolution than what lighting is being calculated for, since the extra resolution is usually not noticeable
-	TextureMapping->IrradiancePhotonCacheSizeX = FMath::Max(FMath::Trunc(TextureMapping->CachedSizeX / PhotonMappingSettings.CachedIrradiancePhotonDownsampleFactor), 6);
-	TextureMapping->IrradiancePhotonCacheSizeY = FMath::Max(FMath::Trunc(TextureMapping->CachedSizeY / PhotonMappingSettings.CachedIrradiancePhotonDownsampleFactor), 6);
+	TextureMapping->IrradiancePhotonCacheSizeX = FMath::Max(FMath::TruncToInt(TextureMapping->CachedSizeX / PhotonMappingSettings.CachedIrradiancePhotonDownsampleFactor), 6);
+	TextureMapping->IrradiancePhotonCacheSizeY = FMath::Max(FMath::TruncToInt(TextureMapping->CachedSizeY / PhotonMappingSettings.CachedIrradiancePhotonDownsampleFactor), 6);
 
 	FTexelToVertexMap TexelToVertexMap(TextureMapping->IrradiancePhotonCacheSizeX, TextureMapping->IrradiancePhotonCacheSizeY);
 
@@ -281,8 +281,8 @@ void FStaticLightingSystem::CacheIrradiancePhotonsTextureMapping(FStaticLighting
 	int32 IrradiancePhotonCacheDebugY = -1;
 	if (bDebugThisMapping)
 	{
-		IrradiancePhotonCacheDebugX = FMath::Trunc(Scene.DebugInput.LocalX / (float)TextureMapping->CachedSizeX * TextureMapping->IrradiancePhotonCacheSizeX);
-		IrradiancePhotonCacheDebugY = FMath::Trunc(Scene.DebugInput.LocalY / (float)TextureMapping->CachedSizeY * TextureMapping->IrradiancePhotonCacheSizeY);
+		IrradiancePhotonCacheDebugX = FMath::TruncToInt(Scene.DebugInput.LocalX / (float)TextureMapping->CachedSizeX * TextureMapping->IrradiancePhotonCacheSizeX);
+		IrradiancePhotonCacheDebugY = FMath::TruncToInt(Scene.DebugInput.LocalY / (float)TextureMapping->CachedSizeY * TextureMapping->IrradiancePhotonCacheSizeY);
 	}
 #endif
 
@@ -1149,7 +1149,7 @@ void FStaticLightingSystem::CalculateDirectLightingTextureMappingFiltered(
 						{
 							if (ShadowMapData(SubX,SubY).bIsMapped)
 							{
-								Visibility += FMath::Trunc(Filter[FilterX][FilterY] * ShadowMapData(SubX,SubY).Visibility);
+								Visibility += FMath::TruncToInt(Filter[FilterX][FilterY] * ShadowMapData(SubX,SubY).Visibility);
 								Coverage += Filter[FilterX][FilterY];
 							}
 						}
@@ -1823,7 +1823,7 @@ void FStaticLightingSystem::CalculateDirectSignedDistanceFieldLightingTextureMap
 		// Choose an upsample factor based on the average texels/world space ratio
 		// The result is that small, high resolution meshes will not upsample as much, since they don't need it, 
 		// But large, low resolution meshes will upsample a lot.
-		const int32 TargetUpsampleFactor = FMath::Trunc(ShadowSettings.ApproximateHighResTexelsPerMaxTransitionDistance / (RightTriangleSide * ShadowSettings.MaxTransitionDistanceWorldSpace));
+		const int32 TargetUpsampleFactor = FMath::TruncToInt(ShadowSettings.ApproximateHighResTexelsPerMaxTransitionDistance / (RightTriangleSide * ShadowSettings.MaxTransitionDistanceWorldSpace));
 		// Round up to the nearest odd factor, so each destination texel has a high resolution source texel at its center
 		// Clamp the upscale factor to be less than 13, since the quality improvements of upsampling higher than that are negligible.
 		UpsampleFactor = FMath::Clamp(TargetUpsampleFactor - TargetUpsampleFactor % 2 + 1, ShadowSettings.MinDistanceFieldUpsampleFactor, 13);
@@ -2251,8 +2251,8 @@ void FStaticLightingSystem::CalculateDirectSignedDistanceFieldLightingTextureMap
 									// Scatter to all distance field texels within MaxTransitionDistanceWorldSpace, rounded up.
 									// This is an approximation to the actual set of distance field texels that are within MaxTransitionDistanceWorldSpace that tends to work out well.
 									// Apply a clamp to avoid a performance cliff with some texels, whose adjacent texel in lightmap space is actually far away in world space
-									const int32 NumLowResScatterTexelsY = FMath::Min(FMath::Trunc(ShadowSettings.MaxTransitionDistanceWorldSpace / (WorldSpacePerHighResTexelY * UpsampleFactor)) + 1, 100);
-									const int32 NumLowResScatterTexelsX = FMath::Min(FMath::Trunc(ShadowSettings.MaxTransitionDistanceWorldSpace / (WorldSpacePerHighResTexelX * UpsampleFactor)) + 1, 100);
+									const int32 NumLowResScatterTexelsY = FMath::Min(FMath::TruncToInt(ShadowSettings.MaxTransitionDistanceWorldSpace / (WorldSpacePerHighResTexelY * UpsampleFactor)) + 1, 100);
+									const int32 NumLowResScatterTexelsX = FMath::Min(FMath::TruncToInt(ShadowSettings.MaxTransitionDistanceWorldSpace / (WorldSpacePerHighResTexelX * UpsampleFactor)) + 1, 100);
 									MappingContext.Stats.NumSignedDistanceFieldScatters++;
 									for (int32 ScatterOffsetY = -NumLowResScatterTexelsY; ScatterOffsetY <= NumLowResScatterTexelsY; ScatterOffsetY++)
 									{
@@ -3415,8 +3415,8 @@ const FIrradiancePhoton* FStaticLightingTextureMapping::GetCachedIrradiancePhoto
 {
 	checkSlow(IrradiancePhotonCacheSizeX > 0 && IrradiancePhotonCacheSizeY > 0);
 	// Clamping is necessary since the UV's may be outside the [0, 1) range
-	const int32 PhotonX = FMath::Clamp(FMath::Trunc(Vertex.TextureCoordinates[1].X * IrradiancePhotonCacheSizeX), 0, IrradiancePhotonCacheSizeX - 1);
-	const int32 PhotonY = FMath::Clamp(FMath::Trunc(Vertex.TextureCoordinates[1].Y * IrradiancePhotonCacheSizeY), 0, IrradiancePhotonCacheSizeY - 1);
+	const int32 PhotonX = FMath::Clamp(FMath::TruncToInt(Vertex.TextureCoordinates[1].X * IrradiancePhotonCacheSizeX), 0, IrradiancePhotonCacheSizeX - 1);
+	const int32 PhotonY = FMath::Clamp(FMath::TruncToInt(Vertex.TextureCoordinates[1].Y * IrradiancePhotonCacheSizeY), 0, IrradiancePhotonCacheSizeY - 1);
 	const int32 PhotonIndex = PhotonY * IrradiancePhotonCacheSizeX + PhotonX;
 
 	const FIrradiancePhoton* ClosestPhoton = CachedIrradiancePhotons[PhotonIndex];
diff --git a/Engine/Source/Programs/UnrealLightmass/Private/Lighting/VolumeDistanceField.cpp b/Engine/Source/Programs/UnrealLightmass/Private/Lighting/VolumeDistanceField.cpp
--- a/Engine/Source/Programs/UnrealLightmass/Private/Lighting/VolumeDistanceField.cpp
+++ b/Engine/Source/Programs/UnrealLightmass/Private/Lighting/VolumeDistanceField.cpp
@@ -26,9 +26,9 @@ void FStaticLightingSystem::BeginCalculateVolumeDistanceField()
 	UnclampedDistanceFieldVolumeBounds.Max = UnclampedDistanceFieldVolumeBounds.Min + DoubleExtent;
 
 	const FVector4 VolumeSizes = UnclampedDistanceFieldVolumeBounds.GetExtent() * 2.0f / VolumeDistanceFieldSettings.VoxelSize;
-	VolumeSizeX = FMath::Trunc(VolumeSizes.X + DELTA);
-	VolumeSizeY = FMath::Trunc(VolumeSizes.Y + DELTA);
-	VolumeSizeZ = FMath::Trunc(VolumeSizes.Z + DELTA);
+	VolumeSizeX = FMath::TruncToFloat(VolumeSizes.X + DELTA);
+	VolumeSizeY = FMath::TruncToFloat(VolumeSizes.Y + DELTA);
+	VolumeSizeZ = FMath::TruncToFloat(VolumeSizes.Z + DELTA);
 
 	// Use a float to avoid 32 bit integer overflow with large volumes
 	const float NumVoxels = VolumeSizeX * VolumeSizeY * VolumeSizeZ;
@@ -49,9 +49,9 @@ void FStaticLightingSystem::BeginCalculateVolumeDistanceField()
 		DistanceFieldVolumeBounds.Max = DistanceFieldVolumeBounds.Min + DoubleExtent;
 
 		const FVector4 ClampedVolumeSizes = DistanceFieldVolumeBounds.GetExtent() * 2.0f / DistanceFieldVoxelSize;
-		VolumeSizeX = FMath::Trunc(ClampedVolumeSizes.X + DELTA);
-		VolumeSizeY = FMath::Trunc(ClampedVolumeSizes.Y + DELTA);
-		VolumeSizeZ = FMath::Trunc(ClampedVolumeSizes.Z + DELTA);
+		VolumeSizeX = FMath::TruncToFloat(ClampedVolumeSizes.X + DELTA);
+		VolumeSizeY = FMath::TruncToFloat(ClampedVolumeSizes.Y + DELTA);
+		VolumeSizeZ = FMath::TruncToFloat(ClampedVolumeSizes.Z + DELTA);
 		
 		LogSolverMessage(FString::Printf(TEXT("CalculateVolumeDistanceField %ux%ux%u, clamped to %ux%ux%u"), OldSizeX, OldSizeY, OldSizeZ, VolumeSizeX, VolumeSizeY, VolumeSizeZ));
 	}
@@ -76,8 +76,8 @@ void FStaticLightingSystem::CalculateVolumeDistanceFieldWorkRange(int32 ZIndex)
 
 	TArray<FVector4> SampleDirections;
 	TArray<FVector2D> SampleDirectionUniforms;
-	const int32 NumThetaSteps = FMath::Trunc(FMath::Sqrt(VolumeDistanceFieldSettings.NumVoxelDistanceSamples / (2.0f * (float)PI)));
-	const int32 NumPhiSteps = FMath::Trunc(NumThetaSteps * (float)PI);
+	const int32 NumThetaSteps = FMath::TruncToInt(FMath::Sqrt(VolumeDistanceFieldSettings.NumVoxelDistanceSamples / (2.0f * (float)PI)));
+	const int32 NumPhiSteps = FMath::TruncToInt(NumThetaSteps * (float)PI);
 	FLMRandomStream RandomStream(0);
 	GenerateStratifiedUniformHemisphereSamples(NumThetaSteps, NumPhiSteps, RandomStream, SampleDirections, SampleDirectionUniforms);
 	TArray<FVector4> OtherHemisphereSamples;
@@ -164,8 +164,8 @@ void FStaticLightingSystem::CalculateVolumeDistanceFieldWorkRange(int32 ZIndex)
 			const float NormalizedDistance1 = FMath::Clamp(MinDistance[1] / VolumeDistanceFieldSettings.VolumeMaxDistance + .5f, 0.0f, 1.0f);
 
 			const FColor FinalValue(
-				FMath::Clamp<uint8>(FMath::Trunc(NormalizedDistance0 * 255), 0, 255), 
-				FMath::Clamp<uint8>(FMath::Trunc(NormalizedDistance1 * 255), 0, 255), 
+				FMath::Clamp<uint8>(FMath::TruncToInt(NormalizedDistance0 * 255), 0, 255), 
+				FMath::Clamp<uint8>(FMath::TruncToInt(NormalizedDistance1 * 255), 0, 255), 
 				Mask0,
 				Mask1
 				);
diff --git a/Engine/Source/Runtime/Core/Private/GenericPlatform/GenericPlatformProcess.cpp b/Engine/Source/Runtime/Core/Private/GenericPlatform/GenericPlatformProcess.cpp
--- a/Engine/Source/Runtime/Core/Private/GenericPlatform/GenericPlatformProcess.cpp
+++ b/Engine/Source/Runtime/Core/Private/GenericPlatform/GenericPlatformProcess.cpp
@@ -224,7 +224,7 @@ void FGenericPlatformProcess::Sleep( float Seconds )
 {
 	SCOPE_CYCLE_COUNTER(STAT_Sleep);
 	FThreadIdleStats::FScopeIdle Scope;
-	const int32 usec = FPlatformMath::Trunc(Seconds * 1000000.0f);
+	const int32 usec = FPlatformMath::TruncToInt(Seconds * 1000000.0f);
 	if (usec > 0)
 	{
 		usleep(usec);
diff --git a/Engine/Source/Runtime/Core/Private/GenericPlatform/GenericPlatformTime.cpp b/Engine/Source/Runtime/Core/Private/GenericPlatform/GenericPlatformTime.cpp
--- a/Engine/Source/Runtime/Core/Private/GenericPlatform/GenericPlatformTime.cpp
+++ b/Engine/Source/Runtime/Core/Private/GenericPlatform/GenericPlatformTime.cpp
@@ -117,31 +117,31 @@ FString FGenericPlatformTime::PrettyTime( double Seconds )
 {
 	if ( Seconds < 1.0 )
 	{
-		return FString::Printf( TEXT("%d ms"), FMath::Trunc(Seconds*1000) );
+		return FString::Printf( TEXT("%d ms"), FMath::TruncToInt(Seconds*1000) );
 	}
 	else if ( Seconds < 10.0 )
 	{
-		int32 Sec = FMath::Trunc(Seconds);
-		int32 Ms = FMath::Trunc(Seconds*1000) - Sec*1000;
+		int32 Sec = FMath::TruncToInt(Seconds);
+		int32 Ms = FMath::TruncToInt(Seconds*1000) - Sec*1000;
 		return FString::Printf( TEXT("%d.%02d sec"), Sec, Ms/10 );
 	}
 	else if ( Seconds < 60.0 )
 	{
-		int32 Sec = FMath::Trunc(Seconds);
-		int32 Ms = FMath::Trunc(Seconds*1000) - Sec*1000;
+		int32 Sec = FMath::TruncToInt(Seconds);
+		int32 Ms = FMath::TruncToInt(Seconds*1000) - Sec*1000;
 		return FString::Printf( TEXT("%d.%d sec"), Sec, Ms/100 );
 	}
 	else if ( Seconds < 60.0*60.0 )
 	{
-		int32 Min = FMath::Trunc(Seconds/60.0);
-		int32 Sec = FMath::Trunc(Seconds) - Min*60;
+		int32 Min = FMath::TruncToInt(Seconds/60.0);
+		int32 Sec = FMath::TruncToInt(Seconds) - Min*60;
 		return FString::Printf( TEXT("%d:%02d min"), Min, Sec );
 	}
 	else
 	{
-		int32 Hr = FMath::Trunc(Seconds/3600.0);
-		int32 Min = FMath::Trunc((Seconds - Hr*3600)/60.0);
-		int32 Sec = FMath::Trunc(Seconds - Hr*3600 - Min*60);
+		int32 Hr = FMath::TruncToInt(Seconds/3600.0);
+		int32 Min = FMath::TruncToInt((Seconds - Hr*3600)/60.0);
+		int32 Sec = FMath::TruncToInt(Seconds - Hr*3600 - Min*60);
 		return FString::Printf( TEXT("%d:%02d:%02d hours"), Hr, Min, Sec );
 	}
 }
diff --git a/Engine/Source/Runtime/Core/Private/Linux/LinuxCursor.cpp b/Engine/Source/Runtime/Core/Private/Linux/LinuxCursor.cpp
--- a/Engine/Source/Runtime/Core/Private/Linux/LinuxCursor.cpp
+++ b/Engine/Source/Runtime/Core/Private/Linux/LinuxCursor.cpp
@@ -191,10 +191,10 @@ void FLinuxCursor::Lock( const RECT* const Bounds )
 	else
 	{
 		SDL_SetWindowGrab( NULL, SDL_TRUE );
-		CursorClipRect.Min.X = FMath::Trunc(Bounds->left);
-		CursorClipRect.Min.Y = FMath::Trunc(Bounds->top);
-		CursorClipRect.Max.X = FMath::Trunc(Bounds->right) - 1;
-		CursorClipRect.Max.Y = FMath::Trunc(Bounds->bottom) - 1;
+		CursorClipRect.Min.X = FMath::TruncToInt(Bounds->left);
+		CursorClipRect.Min.Y = FMath::TruncToInt(Bounds->top);
+		CursorClipRect.Max.X = FMath::TruncToInt(Bounds->right) - 1;
+		CursorClipRect.Max.Y = FMath::TruncToInt(Bounds->bottom) - 1;
 	}
 
 	FVector2D CurrentPosition = GetPosition();
diff --git a/Engine/Source/Runtime/Core/Private/Linux/LinuxWindow.cpp b/Engine/Source/Runtime/Core/Private/Linux/LinuxWindow.cpp
--- a/Engine/Source/Runtime/Core/Private/Linux/LinuxWindow.cpp
+++ b/Engine/Source/Runtime/Core/Private/Linux/LinuxWindow.cpp
@@ -72,10 +72,10 @@ void FLinuxWindow::Initialize( FLinuxApplication* const Application, const TShar
 	const float WidthInitial = Definition->WidthDesiredOnScreen;
 	const float HeightInitial = Definition->HeightDesiredOnScreen;
 
-	int32 X = FMath::Trunc( XInitialRect );
-	int32 Y = FMath::Trunc( YInitialRect );
-	int32 ClientWidth = FMath::Trunc( WidthInitial );
-	int32 ClientHeight = FMath::Trunc( HeightInitial );
+	int32 X = FMath::TruncToInt( XInitialRect );
+	int32 Y = FMath::TruncToInt( YInitialRect );
+	int32 ClientWidth = FMath::TruncToInt( WidthInitial );
+	int32 ClientHeight = FMath::TruncToInt( HeightInitial );
 	int32 WindowWidth = ClientWidth;
 	int32 WindowHeight = ClientHeight;
 
diff --git a/Engine/Source/Runtime/Core/Private/Mac/MacApplication.cpp b/Engine/Source/Runtime/Core/Private/Mac/MacApplication.cpp
--- a/Engine/Source/Runtime/Core/Private/Mac/MacApplication.cpp
+++ b/Engine/Source/Runtime/Core/Private/Mac/MacApplication.cpp
@@ -272,8 +272,8 @@ void FMacApplication::ProcessEvent( NSEvent* Event )
 			// its position based on mouse move delta
 			if( DraggedWindow )
 			{
-				const int32 X = FMath::Trunc(CurrentEventWindow->PositionX + [Event deltaX]);
-				const int32 Y = FMath::Trunc(CurrentEventWindow->PositionY + [Event deltaY]);
+				const int32 X = FMath::TruncToInt(CurrentEventWindow->PositionX + [Event deltaX]);
+				const int32 Y = FMath::TruncToInt(CurrentEventWindow->PositionY + [Event deltaY]);
 				if( CurrentEventWindow.IsValid() )
 				{
 					MessageHandler->OnMovedWindow( CurrentEventWindow.ToSharedRef(), X, Y );
@@ -775,7 +775,7 @@ FPlatformRect FMacApplication::GetWorkArea( const FPlatformRect& CurrentWindow )
 
 	NSScreen* Screen = FindScreenByPoint( CurrentWindow.Left, CurrentWindow.Top );
 
-	const int32 ScreenHeight = FMath::Trunc([Screen frame].size.height);
+	const int32 ScreenHeight = FMath::TruncToInt([Screen frame].size.height);
 	const NSRect VisibleFrame = [Screen visibleFrame];
 
 	FPlatformRect WorkArea;
@@ -974,8 +974,8 @@ void FMacApplication::OnWindowDidResize( FSlateCocoaWindow* Window )
 		if([Window windowMode] != EWindowMode::Windowed)
 		{
 			// Grab current monitor data for sizing
-			Width = FMath::Trunc([[Window screen] frame].size.width);
-			Height = FMath::Trunc([[Window screen] frame].size.height);
+			Width = FMath::TruncToInt([[Window screen] frame].size.width);
+			Height = FMath::TruncToInt([[Window screen] frame].size.height);
 		}
 		
 		MessageHandler->OnSizeChanged( EventWindow.ToSharedRef(), Width, Height );
diff --git a/Engine/Source/Runtime/Core/Private/Mac/MacCursor.cpp b/Engine/Source/Runtime/Core/Private/Mac/MacCursor.cpp
--- a/Engine/Source/Runtime/Core/Private/Mac/MacCursor.cpp
+++ b/Engine/Source/Runtime/Core/Private/Mac/MacCursor.cpp
@@ -152,7 +152,7 @@ FVector2D FMacCursor::GetPosition() const
 	CGPoint CursorPos = CGEventGetLocation(Event);
 	CFRelease(Event);
 	
-	return FVector2D(FMath::Trunc(CursorPos.x), FMath::Trunc(CursorPos.y));
+	return FVector2D(FMath::TruncToFloat(CursorPos.x), FMath::TruncToFloat(CursorPos.y));
 }
 
 void FMacCursor::SetPosition( const int32 X, const int32 Y )
@@ -200,10 +200,10 @@ void FMacCursor::Lock( const RECT* const Bounds )
 	}
 	else
 	{
-		CusorClipRect.Min.X = FMath::Trunc(Bounds->left);
-		CusorClipRect.Min.Y = FMath::Trunc(Bounds->top);
-		CusorClipRect.Max.X = FMath::Trunc(Bounds->right) - 1;
-		CusorClipRect.Max.Y = FMath::Trunc(Bounds->bottom) - 1;
+		CusorClipRect.Min.X = FMath::TruncToInt(Bounds->left);
+		CusorClipRect.Min.Y = FMath::TruncToInt(Bounds->top);
+		CusorClipRect.Max.X = FMath::TruncToInt(Bounds->right) - 1;
+		CusorClipRect.Max.Y = FMath::TruncToInt(Bounds->bottom) - 1;
 	}
 
 	FVector2D CurrentPosition = GetPosition();
@@ -267,7 +267,7 @@ void FMacCursor::UpdateVisibility()
 
 void FMacCursor::WarpCursor( const int32 X, const int32 Y )
 {
-	CGWarpMouseCursorPosition( NSMakePoint( FMath::Trunc( X ), FMath::Trunc( Y ) ) );
+	CGWarpMouseCursorPosition( NSMakePoint( FMath::TruncToInt( X ), FMath::TruncToInt( Y ) ) );
 }
 
 FVector2D FMacCursor::GetMouseWarpDelta( bool const bClearAccumulatedDelta )
diff --git a/Engine/Source/Runtime/Core/Private/Mac/MacWindow.cpp b/Engine/Source/Runtime/Core/Private/Mac/MacWindow.cpp
--- a/Engine/Source/Runtime/Core/Private/Mac/MacWindow.cpp
+++ b/Engine/Source/Runtime/Core/Private/Mac/MacWindow.cpp
@@ -646,19 +646,19 @@ void FMacWindow::Initialize( FMacApplication* const Application, const TSharedRe
 	// Finally, let's initialize the new native window object.  Calling this function will often cause OS
 	// window messages to be sent! (such as activation messages)
 
-	const int32 X = FMath::Trunc( Definition->XDesiredPositionOnScreen );
+	const int32 X = FMath::TruncToInt( Definition->XDesiredPositionOnScreen );
 
 	NSScreen* TargetScreen = Application->FindScreenByPoint( X, Definition->YDesiredPositionOnScreen );
 
-	int32 Y = FMath::Trunc( Definition->YDesiredPositionOnScreen );
+	int32 Y = FMath::TruncToInt( Definition->YDesiredPositionOnScreen );
 
 	// Make sure it's not under the menu bar on whatever display being targeted
-	const int32 ScreenHeight = FMath::Trunc( [TargetScreen frame].size.height );
-	const int32 VisibleHeight = FMath::Trunc( [TargetScreen visibleFrame].origin.y + [TargetScreen visibleFrame].size.height );
+	const int32 ScreenHeight = FMath::TruncToInt( [TargetScreen frame].size.height );
+	const int32 VisibleHeight = FMath::TruncToInt( [TargetScreen visibleFrame].origin.y + [TargetScreen visibleFrame].size.height );
 	Y = FMath::Max( Y, ScreenHeight - VisibleHeight );
 
-	const int32 SizeX = FMath::Max(FMath::Trunc( Definition->WidthDesiredOnScreen ), 1);
-	const int32 SizeY = FMath::Max(FMath::Trunc( Definition->HeightDesiredOnScreen ), 1);
+	const int32 SizeX = FMath::Max(FMath::TruncToInt( Definition->WidthDesiredOnScreen ), 1);
+	const int32 SizeY = FMath::Max(FMath::TruncToInt( Definition->HeightDesiredOnScreen ), 1);
 
 	PositionX = X;
 	PositionY = Y;
diff --git a/Engine/Source/Runtime/Core/Private/Math/Color.cpp b/Engine/Source/Runtime/Core/Private/Math/Color.cpp
--- a/Engine/Source/Runtime/Core/Private/Math/Color.cpp
+++ b/Engine/Source/Runtime/Core/Private/Math/Color.cpp
@@ -69,10 +69,10 @@ FColor FLinearColor::ToRGBE() const
 		int32	Exponent;
 		const float Scale	= frexp(Primary, &Exponent) / Primary * 255.f;
 
-		Color.R		= FMath::Clamp(FMath::Trunc(R * Scale), 0, 255);
-		Color.G		= FMath::Clamp(FMath::Trunc(G * Scale), 0, 255);
-		Color.B		= FMath::Clamp(FMath::Trunc(B * Scale), 0, 255);
-		Color.A		= FMath::Clamp(FMath::Trunc(Exponent),-128,127) + 128;
+		Color.R		= FMath::Clamp(FMath::TruncToInt(R * Scale), 0, 255);
+		Color.G		= FMath::Clamp(FMath::TruncToInt(G * Scale), 0, 255);
+		Color.B		= FMath::Clamp(FMath::TruncToInt(B * Scale), 0, 255);
+		Color.A		= FMath::Clamp(FMath::TruncToInt(Exponent),-128,127) + 128;
 	}
 
 	return Color;
@@ -96,10 +96,10 @@ FColor FLinearColor::ToFColor(const bool bSRGB) const
 
 	FColor ret;
 
-	ret.A = FMath::Floor(FloatA * 255.999f);
-	ret.R = FMath::Floor(FloatR * 255.999f);
-	ret.G = FMath::Floor(FloatG * 255.999f);
-	ret.B = FMath::Floor(FloatB * 255.999f);
+	ret.A = FMath::FloorToInt(FloatA * 255.999f);
+	ret.R = FMath::FloorToInt(FloatR * 255.999f);
+	ret.G = FMath::FloorToInt(FloatG * 255.999f);
+	ret.B = FMath::FloorToInt(FloatB * 255.999f);
 
 	return ret;
 }
@@ -108,10 +108,10 @@ FColor FLinearColor::ToFColor(const bool bSRGB) const
 FColor FLinearColor::Quantize() const
 {
 	return FColor(
-		(uint8)FMath::Clamp<int32>(FMath::Trunc(R*255.f),0,255),
-		(uint8)FMath::Clamp<int32>(FMath::Trunc(G*255.f),0,255),
-		(uint8)FMath::Clamp<int32>(FMath::Trunc(B*255.f),0,255),
-		(uint8)FMath::Clamp<int32>(FMath::Trunc(A*255.f),0,255)
+		(uint8)FMath::Clamp<int32>(FMath::TruncToInt(R*255.f),0,255),
+		(uint8)FMath::Clamp<int32>(FMath::TruncToInt(G*255.f),0,255),
+		(uint8)FMath::Clamp<int32>(FMath::TruncToInt(B*255.f),0,255),
+		(uint8)FMath::Clamp<int32>(FMath::TruncToInt(A*255.f),0,255)
 		);
 }
 
@@ -272,8 +272,8 @@ FColor FColor::MakeRedToGreenColorFromScalar(float Scalar)
 
 	float RedSclr = FMath::Clamp<float>((1.0f - Scalar)/0.5f,0.f,1.f);
 	float GreenSclr = FMath::Clamp<float>((Scalar/0.5f),0.f,1.f);
-	R = FMath::Trunc(255 * RedSclr);
-	G = FMath::Trunc(255 * GreenSclr);
+	R = FMath::TruncToInt(255 * RedSclr);
+	G = FMath::TruncToInt(255 * GreenSclr);
 	return FColor(R,G,B);
 }
 
diff --git a/Engine/Source/Runtime/Core/Private/Misc/CoreMisc.cpp b/Engine/Source/Runtime/Core/Private/Misc/CoreMisc.cpp
--- a/Engine/Source/Runtime/Core/Private/Misc/CoreMisc.cpp
+++ b/Engine/Source/Runtime/Core/Private/Misc/CoreMisc.cpp
@@ -936,12 +936,12 @@ void GenerateConvenientWindowedResolutions(const FDisplayMetrics& InDisplayMetri
 
 			if (bInPortraitMode)
 			{
-				TargetHeight = FMath::Round(InDisplayMetrics.PrimaryDisplayHeight * Scales[ScaleIndex]);
+				TargetHeight = FMath::RoundToFloat(InDisplayMetrics.PrimaryDisplayHeight * Scales[ScaleIndex]);
 				TargetWidth = TargetHeight * Aspect;
 			}
 			else
 			{
-				TargetWidth = FMath::Round(InDisplayMetrics.PrimaryDisplayWidth * Scales[ScaleIndex]);
+				TargetWidth = FMath::RoundToFloat(InDisplayMetrics.PrimaryDisplayWidth * Scales[ScaleIndex]);
 				TargetHeight = TargetWidth * Aspect;
 			}
 
diff --git a/Engine/Source/Runtime/Core/Private/Misc/DateTime.cpp b/Engine/Source/Runtime/Core/Private/Misc/DateTime.cpp
--- a/Engine/Source/Runtime/Core/Private/Misc/DateTime.cpp
+++ b/Engine/Source/Runtime/Core/Private/Misc/DateTime.cpp
@@ -113,7 +113,7 @@ FDateTime::FDate FDateTime::ToDate( ) const
 
 	int32 i, j, k, l, n;
 
-	l = FMath::Floor(GetJulianDay() + 0.5) + 68569;
+	l = FMath::FloorToInt(GetJulianDay() + 0.5) + 68569;
 	n = 4 * l / 146097;
 	l = l - (146097 * n + 3) / 4;
 	i = 4000 * (l + 1) / 1461001;
diff --git a/Engine/Source/Runtime/Core/Private/Windows/WindowsApplication.cpp b/Engine/Source/Runtime/Core/Private/Windows/WindowsApplication.cpp
--- a/Engine/Source/Runtime/Core/Private/Windows/WindowsApplication.cpp
+++ b/Engine/Source/Runtime/Core/Private/Windows/WindowsApplication.cpp
@@ -202,8 +202,8 @@ void FWindowsApplication::SetHighPrecisionMouseMode( const bool Enable, const TS
 bool FWindowsApplication::TryCalculatePopupWindowPosition( const FPlatformRect& InAnchor, const FVector2D& InSize, const EPopUpOrientation::Type Orientation, /*OUT*/ FVector2D* const CalculatedPopUpPosition ) const
 {
 #if(_WIN32_WINNT >= 0x0601) 
-	POINT AnchorPoint = { FMath::Trunc(InAnchor.Left), FMath::Trunc(InAnchor.Top) };
-	SIZE WindowSize = { FMath::Trunc(InSize.X), FMath::Trunc(InSize.Y) };
+	POINT AnchorPoint = { FMath::TruncToInt(InAnchor.Left), FMath::TruncToInt(InAnchor.Top) };
+	SIZE WindowSize = { FMath::TruncToInt(InSize.X), FMath::TruncToInt(InSize.Y) };
 
 	RECT WindowPos;
 	::CalculatePopupWindowPosition( &AnchorPoint, &WindowSize, TPM_CENTERALIGN | TPM_VCENTERALIGN, NULL, &WindowPos );
diff --git a/Engine/Source/Runtime/Core/Private/Windows/WindowsWindow.cpp b/Engine/Source/Runtime/Core/Private/Windows/WindowsWindow.cpp
--- a/Engine/Source/Runtime/Core/Private/Windows/WindowsWindow.cpp
+++ b/Engine/Source/Runtime/Core/Private/Windows/WindowsWindow.cpp
@@ -50,10 +50,10 @@ void FWindowsWindow::Initialize( FWindowsApplication* const Application, const T
 	const float WidthInitial = Definition->WidthDesiredOnScreen;
 	const float HeightInitial = Definition->HeightDesiredOnScreen;
 
-	int32 X = FMath::Trunc( XInitialRect );
-	int32 Y = FMath::Trunc( YInitialRect );
-	int32 ClientWidth = FMath::Trunc( WidthInitial );
-	int32 ClientHeight = FMath::Trunc( HeightInitial );
+	int32 X = FMath::TruncToInt( XInitialRect );
+	int32 Y = FMath::TruncToInt( YInitialRect );
+	int32 ClientWidth = FMath::TruncToInt( WidthInitial );
+	int32 ClientHeight = FMath::TruncToInt( HeightInitial );
 	int32 WindowWidth = ClientWidth;
 	int32 WindowHeight = ClientHeight;
 
@@ -598,7 +598,7 @@ void FWindowsWindow::SetWindowFocus()
  */
 void FWindowsWindow::SetOpacity( const float InOpacity )
 {
-	SetLayeredWindowAttributes( HWnd, 0, FMath::Trunc( InOpacity * 255.0f ), LWA_ALPHA );
+	SetLayeredWindowAttributes( HWnd, 0, FMath::TruncToInt( InOpacity * 255.0f ), LWA_ALPHA );
 }
 
 /**
diff --git a/Engine/Source/Runtime/Core/Public/GenericPlatform/GenericPlatformMath.h b/Engine/Source/Runtime/Core/Public/GenericPlatform/GenericPlatformMath.h
--- a/Engine/Source/Runtime/Core/Public/GenericPlatform/GenericPlatformMath.h
+++ b/Engine/Source/Runtime/Core/Public/GenericPlatform/GenericPlatformMath.h
@@ -18,7 +18,7 @@ struct FGenericPlatformMath
 	 * @param F		Floating point value to convert
 	 * @return		Truncated integer.
 	 */
-	static CONSTEXPR FORCEINLINE int32 Trunc( float F )
+	static CONSTEXPR FORCEINLINE int32 TruncToInt(float F)
 	{
 		return (int32)F;
 	}
@@ -28,27 +28,37 @@ struct FGenericPlatformMath
 	 * @param F		Floating point value to convert
 	 * @return		Truncated integer value.
 	 */
-	static CONSTEXPR FORCEINLINE float TruncFloat( float F )
+	static CONSTEXPR FORCEINLINE float TruncToFloat(float F)
 	{
-		return (float)Trunc(F);
+		return (float)TruncToInt(F);
 	}
 
 	/**
-	 * Converts a float to a less or equal integer.
+	 * Converts a float to a nearest less or equal integer.
 	 * @param F		Floating point value to convert
 	 * @return		An integer less or equal to 'F'.
 	 */
-	static FORCEINLINE int32 Floor( float F )
+	static FORCEINLINE int32 FloorToInt(float F)
 	{
-		return Trunc(floorf(F));
+		return TruncToInt(floorf(F));
 	}
 
 	/**
-	 * Converts a double to a less or equal integer.
+	* Converts a float to the nearest less or equal integer.
+	* @param F		Floating point value to convert
+	* @return		An integer less or equal to 'F'.
+	*/
+	static FORCEINLINE float FloorToFloat(float F)
+	{
+		return floorf(F);
+	}
+
+	/**
+	 * Converts a double to the nearest less or equal integer.
 	 * @param F		Floating point value to convert
-	 * @return		The nearest integer value to 'F'.
+	 * @return		An integer less or equal to 'F'.
 	 */
-	static FORCEINLINE double FloorDouble( double F )
+	static FORCEINLINE double FloorDouble(double F)
 	{
 		return floor(F);
 	}
@@ -58,28 +68,59 @@ struct FGenericPlatformMath
 	 * @param F		Floating point value to convert
 	 * @return		The nearest integer to 'F'.
 	 */
-	static FORCEINLINE int32 Round( float F )
+	static FORCEINLINE int32 RoundToInt(float F)
 	{
-		return Floor(F+0.5f);
+		return FloorToInt(F + 0.5f);
 	}
 
 	/**
-	 * Converts a float to a greater or equal integer.
-	 * @param F		Floating point value to convert
-	 * @return		An integer greater or equal to 'F'.
-	 */
-	static FORCEINLINE int32 Ceil( float F )
+	* Converts a float to the nearest integer. Rounds up when the fraction is .5
+	* @param F		Floating point value to convert
+	* @return		The nearest integer to 'F'.
+	*/
+	static FORCEINLINE float RoundToFloat(float F)
 	{
-		return Trunc(ceilf(F));
+		return FloorToFloat(F + 0.5f);
 	}
+
 	/**
-	 * Returns the fractional part of a float.
-	 * @param Value	Floating point value to convert
-	 * @return		A float between >=0 and < 1.
-	 */
-	static FORCEINLINE float Fractional( float Value ) 
-	{ 
-		return Value - TruncFloat( Value );
+	* Converts a float to the nearest greater or equal integer.
+	* @param F		Floating point value to convert
+	* @return		An integer greater or equal to 'F'.
+	*/
+	static FORCEINLINE int32 CeilToInt(float F)
+	{
+		return TruncToInt(ceilf(F));
+	}
+
+	/**
+	* Converts a float to the nearest greater or equal integer.
+	* @param F		Floating point value to convert
+	* @return		An integer greater or equal to 'F'.
+	*/
+	static FORCEINLINE float CeilToFloat(float F)
+	{
+		return ceilf(F);
+	}
+
+	/**
+	* Returns signed fractional part of a float.
+	* @param Value	Floating point value to convert
+	* @return		A float between >=0 and < 1 for nonnegative input. A float between >= -1 and < 0 for negative input.
+	*/
+	static FORCEINLINE float Fractional(float Value)
+	{
+		return Value - TruncToFloat(Value);
+	}
+
+	/**
+	* Returns the fractional part of a float.
+	* @param Value	Floating point value to convert
+	* @return		A float between >=0 and < 1.
+	*/
+	static FORCEINLINE float Frac(float Value)
+	{
+		return Value - FloorToFloat(Value);
 	}
 
 	static FORCEINLINE float Exp( float Value ) { return expf(Value); }
diff --git a/Engine/Source/Runtime/Core/Public/Linux/LinuxPlatformMath.h b/Engine/Source/Runtime/Core/Public/Linux/LinuxPlatformMath.h
--- a/Engine/Source/Runtime/Core/Public/Linux/LinuxPlatformMath.h
+++ b/Engine/Source/Runtime/Core/Public/Linux/LinuxPlatformMath.h
@@ -12,28 +12,49 @@
 struct FLinuxPlatformMath : public FGenericPlatformMath
 {
 #if PLATFORM_ENABLE_VECTORINTRINSICS
-	static FORCEINLINE int32 Trunc( float F )
+	static FORCEINLINE int32 TruncToInt(float F)
 	{
-		return _mm_cvtt_ss2si( _mm_set_ss( F ) );
+		return _mm_cvtt_ss2si(_mm_set_ss(F));
 	}
-	static FORCEINLINE float TruncFloat( float F )
+
+	static FORCEINLINE float TruncToFloat(float F)
 	{
-		return (float)Trunc(F); // same as generic implementation, but this will call the faster trunc
+		return (float)TruncToInt(F); // same as generic implementation, but this will call the faster trunc
 	}
-	static FORCEINLINE int32 Round( float F )
+
+	static FORCEINLINE int32 RoundToInt(float F)
 	{
 		// Note: the x2 is to workaround the rounding-to-nearest-even-number issue when the fraction is .5
-		return _mm_cvt_ss2si( _mm_set_ss(F + F + 0.5f) ) >> 1;
+		return _mm_cvt_ss2si(_mm_set_ss(F + F + 0.5f)) >> 1;
+	}
+
+	static FORCEINLINE float RoundToFloat(float F)
+	{
+		return (float)RoundToInt(F);
 	}
-	static FORCEINLINE int32 Floor( float F )
+
+	static FORCEINLINE int32 FloorToInt(float F)
 	{
-		return _mm_cvt_ss2si( _mm_set_ss(F + F - 0.5f) ) >> 1;
+		return _mm_cvt_ss2si(_mm_set_ss(F + F - 0.5f)) >> 1;
 	}
-	static FORCEINLINE int32 Ceil( float F )
+
+	static FORCEINLINE float FloorToFloat(float F)
+	{
+		return (float)FloorToInt(F);
+	}
+
+	static FORCEINLINE int32 CeilToInt(float F)
 	{
 		// Note: the x2 is to workaround the rounding-to-nearest-even-number issue when the fraction is .5
-		return -(_mm_cvt_ss2si( _mm_set_ss(-0.5f - (F + F))) >> 1);
+		return -(_mm_cvt_ss2si(_mm_set_ss(-0.5f - (F + F))) >> 1);
 	}
+
+	static FORCEINLINE float CeilToFloat(float F)
+	{
+		// Note: the x2 is to workaround the rounding-to-nearest-even-number issue when the fraction is .5
+		return (float)CeilToInt(F);
+	}
+
 	static FORCEINLINE bool IsNaN( float A ) { return isnan(A) != 0; }
 	static FORCEINLINE bool IsFinite( float A ) { return isfinite(A); }
 
diff --git a/Engine/Source/Runtime/Core/Public/Mac/MacPlatformMath.h b/Engine/Source/Runtime/Core/Public/Mac/MacPlatformMath.h
--- a/Engine/Source/Runtime/Core/Public/Mac/MacPlatformMath.h
+++ b/Engine/Source/Runtime/Core/Public/Mac/MacPlatformMath.h
@@ -11,28 +11,48 @@
 **/
 struct FMacPlatformMath : public FGenericPlatformMath
 {
-	static FORCEINLINE int32 Trunc( float F )
+	static FORCEINLINE int32 TruncToInt(float F)
 	{
-		return _mm_cvtt_ss2si( _mm_set_ss( F ) );
+		return _mm_cvtt_ss2si(_mm_set_ss(F));
 	}
-	static FORCEINLINE float TruncFloat( float F )
+
+	static FORCEINLINE float TruncFloat(float F)
 	{
-		return (float)Trunc(F); // same as generic implementation, but this will call the faster trunc
+		return (float)TruncToInt(F); // same as generic implementation, but this will call the faster trunc
 	}
-	static FORCEINLINE int32 Round( float F )
+
+	static FORCEINLINE int32 RoundToInt(float F)
 	{
 		// Note: the x2 is to workaround the rounding-to-nearest-even-number issue when the fraction is .5
-		return _mm_cvt_ss2si( _mm_set_ss(F + F + 0.5f) ) >> 1;
+		return _mm_cvt_ss2si(_mm_set_ss(F + F + 0.5f)) >> 1;
+	}
+
+	static FORCEINLINE float RoundToFloat(float F)
+	{
+		return (float)RoundToInt(F);
+	}
+
+	static FORCEINLINE int32 FloorToInt(float F)
+	{
+		return _mm_cvt_ss2si(_mm_set_ss(F + F - 0.5f)) >> 1;
 	}
-	static FORCEINLINE int32 Floor( float F )
+
+	static FORCEINLINE float FloorToFloat(float F)
 	{
-		return _mm_cvt_ss2si( _mm_set_ss(F + F - 0.5f) ) >> 1;
+		return (float)FloorToInt(F);
 	}
-	static FORCEINLINE int32 Ceil( float F )
+
+	static FORCEINLINE int32 CeilToInt(float F)
 	{
 		// Note: the x2 is to workaround the rounding-to-nearest-even-number issue when the fraction is .5
-		return -(_mm_cvt_ss2si( _mm_set_ss(-0.5f - (F + F))) >> 1);
+		return -(_mm_cvt_ss2si(_mm_set_ss(-0.5f - (F + F))) >> 1);
 	}
+
+	static FORCEINLINE float CeilToFloat(float F)
+	{
+		return (float)CeilToInt(F);
+	}
+
 	static FORCEINLINE bool IsNaN( float A ) { return isnan(A) != 0; }
 	static FORCEINLINE bool IsFinite( float A ) { return isfinite(A); }
 
diff --git a/Engine/Source/Runtime/Core/Public/Math/Color.h b/Engine/Source/Runtime/Core/Public/Math/Color.h
--- a/Engine/Source/Runtime/Core/Public/Math/Color.h
+++ b/Engine/Source/Runtime/Core/Public/Math/Color.h
@@ -351,10 +351,10 @@ public:
 	FColor(const FLinearColor& C)
 		// put these into the body for proper ordering with INTEL vs non-INTEL_BYTE_ORDER
 	{
-		R = FMath::Clamp(FMath::Trunc(FMath::Pow(C.R,1.0f / 2.2f) * 255.0f),0,255);
-		G = FMath::Clamp(FMath::Trunc(FMath::Pow(C.G,1.0f / 2.2f) * 255.0f),0,255);
-		B = FMath::Clamp(FMath::Trunc(FMath::Pow(C.B,1.0f / 2.2f) * 255.0f),0,255);
-		A = FMath::Clamp(FMath::Trunc(       C.A              * 255.0f),0,255);
+		R = FMath::Clamp(FMath::TruncToInt(FMath::Pow(C.R,1.0f / 2.2f) * 255.0f),0,255);
+		G = FMath::Clamp(FMath::TruncToInt(FMath::Pow(C.G,1.0f / 2.2f) * 255.0f),0,255);
+		B = FMath::Clamp(FMath::TruncToInt(FMath::Pow(C.B,1.0f / 2.2f) * 255.0f),0,255);
+		A = FMath::Clamp(FMath::TruncToInt(       C.A              * 255.0f),0,255);
 	}
 
 	FORCEINLINE explicit FColor( uint32 InColor )
diff --git a/Engine/Source/Runtime/Core/Public/Math/Float16.h b/Engine/Source/Runtime/Core/Public/Math/Float16.h
--- a/Engine/Source/Runtime/Core/Public/Math/Float16.h
+++ b/Engine/Source/Runtime/Core/Public/Math/Float16.h
@@ -155,7 +155,7 @@ FORCEINLINE float FFloat16::GetFloat() const
 		else
 		{
 			// Denormal.
-			uint32 MantissaShift = 10 - (uint32)FMath::Trunc(FMath::Log2(Mantissa));
+			uint32 MantissaShift = 10 - (uint32)FMath::TruncToInt(FMath::Log2(Mantissa));
 			Result.Components.Exponent = 127 - (15 - 1) - MantissaShift;
 			Result.Components.Mantissa = Mantissa << (MantissaShift + 23 - 10);
 		}
diff --git a/Engine/Source/Runtime/Core/Public/Math/IntRect.h b/Engine/Source/Runtime/Core/Public/Math/IntRect.h
--- a/Engine/Source/Runtime/Core/Public/Math/IntRect.h
+++ b/Engine/Source/Runtime/Core/Public/Math/IntRect.h
@@ -324,7 +324,7 @@ FORCEINLINE FIntRect FIntRect::Scale( float Fraction ) const
 	FVector2D Min2D = FVector2D(Min.X, Min.Y) * Fraction;
 	FVector2D Max2D = FVector2D(Max.X, Max.Y) * Fraction;
 
-	return FIntRect(FMath::Floor(Min2D.X), FMath::Floor(Min2D.Y), FMath::Ceil(Max2D.X), FMath::Ceil(Max2D.Y));
+	return FIntRect(FMath::FloorToInt(Min2D.X), FMath::FloorToInt(Min2D.Y), FMath::CeilToInt(Max2D.X), FMath::CeilToInt(Max2D.Y));
 }
 
 
diff --git a/Engine/Source/Runtime/Core/Public/Math/Plane.h b/Engine/Source/Runtime/Core/Public/Math/Plane.h
--- a/Engine/Source/Runtime/Core/Public/Math/Plane.h
+++ b/Engine/Source/Runtime/Core/Public/Math/Plane.h
@@ -274,10 +274,10 @@ public:
 		}
 		else
 		{
-			int16 iX(FMath::Round(X));
-			int16 iY(FMath::Round(Y));
-			int16 iZ(FMath::Round(Z));
-			int16 iW(FMath::Round(W));
+			int16 iX(FMath::RoundToInt(X));
+			int16 iY(FMath::RoundToInt(Y));
+			int16 iZ(FMath::RoundToInt(Z));
+			int16 iW(FMath::RoundToInt(W));
 			Ar << iX << iY << iZ << iW;
 		}
 		bOutSuccess = true;
diff --git a/Engine/Source/Runtime/Core/Public/Math/RandomStream.h b/Engine/Source/Runtime/Core/Public/Math/RandomStream.h
--- a/Engine/Source/Runtime/Core/Public/Math/RandomStream.h
+++ b/Engine/Source/Runtime/Core/Public/Math/RandomStream.h
@@ -149,7 +149,7 @@ public:
 	FORCEINLINE int32 RandHelper(int32 A) const
 	{
 		// Can't just multiply GetFraction by A, as GetFraction could be == 1.0f
-		return ((A > 0) ? FMath::Trunc(GetFraction() * ((float)A - DELTA)) : 0);
+		return ((A > 0) ? FMath::TruncToInt(GetFraction() * ((float)A - DELTA)) : 0);
 	}
 
 	/**
diff --git a/Engine/Source/Runtime/Core/Public/Math/Rotator.h b/Engine/Source/Runtime/Core/Public/Math/Rotator.h
--- a/Engine/Source/Runtime/Core/Public/Math/Rotator.h
+++ b/Engine/Source/Runtime/Core/Public/Math/Rotator.h
@@ -506,7 +506,7 @@ FORCEINLINE float FRotator::NormalizeAxis(float Angle)
 FORCEINLINE uint8 FRotator::CompressAxisToByte(float Angle)
 {
 	// map [0->360) to [0->256) and mask off any winding
-	return FMath::Round(Angle * 256.f / 360.f) & 0xFF;
+	return FMath::RoundToInt(Angle * 256.f / 360.f) & 0xFF;
 }
 
 FORCEINLINE float FRotator::DecompressAxisFromByte(uint16 Angle)
@@ -518,7 +518,7 @@ FORCEINLINE float FRotator::DecompressAxisFromByte(uint16 Angle)
 FORCEINLINE uint16 FRotator::CompressAxisToShort(float Angle)
 {
 	// map [0->360) to [0->65536) and mask off any winding
-	return FMath::Round(Angle * 65536.f / 360.f) & 0xFFFF;
+	return FMath::RoundToInt(Angle * 65536.f / 360.f) & 0xFFFF;
 }
 
 
diff --git a/Engine/Source/Runtime/Core/Public/Math/UnrealMathUtility.h b/Engine/Source/Runtime/Core/Public/Math/UnrealMathUtility.h
--- a/Engine/Source/Runtime/Core/Public/Math/UnrealMathUtility.h
+++ b/Engine/Source/Runtime/Core/Public/Math/UnrealMathUtility.h
@@ -90,7 +90,7 @@ struct FMath : public FPlatformMath
 	static FORCEINLINE int32 RandHelper(int32 A)
 	{
 		// RAND_MAX+1 give interval [0..A) with even distribution.
-		return A>0 ? Trunc(Rand()/(float)((uint32)RAND_MAX+1) * A) : 0;
+		return A>0 ? TruncToInt(Rand()/(float)((uint32)RAND_MAX+1) * A) : 0;
 	}
 
 	/** Helper function for rand implementations. Returns a random number >= Min and <= Max */
@@ -231,7 +231,7 @@ struct FMath : public FPlatformMath
 		if( Grid==0.f )	return Location;
 		else			
 		{
-			return Floor((Location + 0.5*Grid)/Grid)*Grid;
+			return FloorToFloat((Location + 0.5*Grid)/Grid)*Grid;
 		}
 	}
 
diff --git a/Engine/Source/Runtime/Core/Public/Math/Vector2D.h b/Engine/Source/Runtime/Core/Public/Math/Vector2D.h
--- a/Engine/Source/Runtime/Core/Public/Math/Vector2D.h
+++ b/Engine/Source/Runtime/Core/Public/Math/Vector2D.h
@@ -719,7 +719,7 @@ FORCEINLINE float FVector2D::Component( int32 Index ) const
 
 FORCEINLINE FIntPoint FVector2D::IntPoint() const
 {
-	return FIntPoint( FMath::Round(X), FMath::Round(Y) );
+	return FIntPoint( FMath::RoundToInt(X), FMath::RoundToInt(Y) );
 }
 
 FORCEINLINE FVector2D FVector2D::ClampAxes( float MinAxisVal, float MaxAxisVal ) const
diff --git a/Engine/Source/Runtime/Core/Public/Windows/WindowsPlatformMath.h b/Engine/Source/Runtime/Core/Public/Windows/WindowsPlatformMath.h
--- a/Engine/Source/Runtime/Core/Public/Windows/WindowsPlatformMath.h
+++ b/Engine/Source/Runtime/Core/Public/Windows/WindowsPlatformMath.h
@@ -13,28 +13,49 @@
 struct FWindowsPlatformMath : public FGenericPlatformMath
 {
 #if PLATFORM_ENABLE_VECTORINTRINSICS
-	static FORCEINLINE int32 Trunc( float F )
+	static FORCEINLINE int32 TruncToInt(float F)
 	{
-		return _mm_cvtt_ss2si( _mm_set_ss( F ) );
+		return _mm_cvtt_ss2si(_mm_set_ss(F));
 	}
+
 	static FORCEINLINE float TruncFloat( float F )
 	{
-		return (float)Trunc(F); // same as generic implementation, but this will call the faster trunc
+		return (float)TruncToInt(F); // same as generic implementation, but this will call the faster trunc
 	}
-	static FORCEINLINE int32 Round( float F )
+
+	static FORCEINLINE int32 RoundToInt( float F )
 	{
 		// Note: the x2 is to workaround the rounding-to-nearest-even-number issue when the fraction is .5
 		return _mm_cvt_ss2si( _mm_set_ss(F + F + 0.5f) ) >> 1;
 	}
-	static FORCEINLINE int32 Floor( float F )
+
+	static FORCEINLINE float RoundToFloat(float F)
+	{
+		return (float)RoundToInt(F);
+	}
+
+	static FORCEINLINE int32 FloorToInt(float F)
 	{
-		return _mm_cvt_ss2si( _mm_set_ss(F + F - 0.5f) ) >> 1;
+		return _mm_cvt_ss2si(_mm_set_ss(F + F - 0.5f)) >> 1;
 	}
-	static FORCEINLINE int32 Ceil( float F )
+
+	static FORCEINLINE float FloorToFloat(float F)
+	{
+		return (float)FloorToInt(F);
+	}
+
+	static FORCEINLINE int32 CeilToInt(float F)
 	{
 		// Note: the x2 is to workaround the rounding-to-nearest-even-number issue when the fraction is .5
-		return -(_mm_cvt_ss2si( _mm_set_ss(-0.5f - (F + F))) >> 1);
+		return -(_mm_cvt_ss2si(_mm_set_ss(-0.5f - (F + F))) >> 1);
 	}
+
+	static FORCEINLINE float CeilToFloat(float F)
+	{
+		// Note: the x2 is to workaround the rounding-to-nearest-even-number issue when the fraction is .5
+		return (float)CeilToInt(F);
+	}
+
 	static FORCEINLINE bool IsNaN( float A ) { return _isnan(A) != 0; }
 	static FORCEINLINE bool IsFinite( float A ) { return _finite(A) != 0; }
 
diff --git a/Engine/Source/Runtime/Engine/Classes/Engine/Font.h b/Engine/Source/Runtime/Engine/Classes/Engine/Font.h
--- a/Engine/Source/Runtime/Engine/Classes/Engine/Font.h
+++ b/Engine/Source/Runtime/Engine/Classes/Engine/Font.h
@@ -222,7 +222,7 @@ public:
 			Total += Width;
 		}
 
-		return( FMath::Ceil( Total ) );
+		return( FMath::CeilToInt( Total ) );
 	}
 
 	/**
@@ -243,7 +243,7 @@ public:
 			Total = FMath::Max( Total, Height );
 		}
 
-		return( FMath::Ceil( Total ) );
+		return( FMath::CeilToInt( Total ) );
 	}
 
 	// Begin UObject interface
diff --git a/Engine/Source/Runtime/Engine/Classes/Engine/NetSerialization.h b/Engine/Source/Runtime/Engine/Classes/Engine/NetSerialization.h
--- a/Engine/Source/Runtime/Engine/Classes/Engine/NetSerialization.h
+++ b/Engine/Source/Runtime/Engine/Classes/Engine/NetSerialization.h
@@ -723,9 +723,9 @@ bool WritePackedVector(FVector Value, FArchive& Ar)	// Note Value is intended to
 	Value *= ScaleFactor;
 
 	// Do basically FVector::SerializeCompressed
-	int32 IntX	= FMath::Round(Value.X);
-	int32 IntY	= FMath::Round(Value.Y);
-	int32 IntZ	= FMath::Round(Value.Z);
+	int32 IntX	= FMath::RoundToInt(Value.X);
+	int32 IntY	= FMath::RoundToInt(Value.Y);
+	int32 IntZ	= FMath::RoundToInt(Value.Z);
 			
 	uint32 Bits	= FMath::Clamp<uint32>( FMath::CeilLogTwo( 1 + FMath::Max3( FMath::Abs(IntX), FMath::Abs(IntY), FMath::Abs(IntZ) ) ), 1, MaxBitsPerComponent ) - 1;
 
@@ -814,13 +814,13 @@ bool WriteFixedCompressedFloat(const float Value, FArchive& Ar)
 	{
 		// We have to scale this down, scale needs to be a float:
 		const float scale = (float)MaxBitValue / (float)MaxValue;
-		ScaledValue = FMath::Trunc(scale * Value);
+		ScaledValue = FMath::TruncToInt(scale * Value);
 	}
 	else
 	{
 		// We will scale up to get extra precision. But keep is a whole number preserve whole values
 		enum { scale = MaxBitValue / MaxValue };
-		ScaledValue = FMath::Round( scale * Value );
+		ScaledValue = FMath::RoundToInt( scale * Value );
 	}
 
 	uint32 Delta = static_cast<uint32>(ScaledValue + Bias);
diff --git a/Engine/Source/Runtime/Engine/Private/AI/EnvironmentQuery/Generators/EnvQueryGenerator_OnCircle.cpp b/Engine/Source/Runtime/Engine/Private/AI/EnvironmentQuery/Generators/EnvQueryGenerator_OnCircle.cpp
--- a/Engine/Source/Runtime/Engine/Private/AI/EnvironmentQuery/Generators/EnvQueryGenerator_OnCircle.cpp
+++ b/Engine/Source/Runtime/Engine/Private/AI/EnvironmentQuery/Generators/EnvQueryGenerator_OnCircle.cpp
@@ -100,7 +100,7 @@ void UEnvQueryGenerator_OnCircle::GenerateItems(FEnvQueryInstance& QueryInstance
 	const float CircumferenceLength = 2.f * PI * RadiusValue;
 	const float ArcAnglePercentage = Angle.Value / 360.f;
 	const float ArcLength = CircumferenceLength * ArcAnglePercentage;
-	const int32 StepsCount = FMath::Ceil(ArcLength / ItemSpace) + 1;
+	const int32 StepsCount = FMath::CeilToInt(ArcLength / ItemSpace) + 1;
 	const float AngleStep = AngleDegree / (StepsCount - 1);
 
 	FVector StartDirection = CalcDirection(QueryInstance);
diff --git a/Engine/Source/Runtime/Engine/Private/AI/EnvironmentQuery/Generators/EnvQueryGenerator_PathingGrid.cpp b/Engine/Source/Runtime/Engine/Private/AI/EnvironmentQuery/Generators/EnvQueryGenerator_PathingGrid.cpp
--- a/Engine/Source/Runtime/Engine/Private/AI/EnvironmentQuery/Generators/EnvQueryGenerator_PathingGrid.cpp
+++ b/Engine/Source/Runtime/Engine/Private/AI/EnvironmentQuery/Generators/EnvQueryGenerator_PathingGrid.cpp
@@ -34,7 +34,7 @@ void UEnvQueryGenerator_PathingGrid::GenerateItems(struct FEnvQueryInstance& Que
 		return;
 	}
 
-	const int32 ItemCount = FPlatformMath::Trunc((PathDistanceValue * 2.0f / DensityValue) + 1);
+	const int32 ItemCount = FPlatformMath::TruncToInt((PathDistanceValue * 2.0f / DensityValue) + 1);
 	const int32 ItemCountHalf = ItemCount / 2;
 
 	TArray<FVector> ContextLocations;
diff --git a/Engine/Source/Runtime/Engine/Private/AI/EnvironmentQuery/Generators/EnvQueryGenerator_SimpleGrid.cpp b/Engine/Source/Runtime/Engine/Private/AI/EnvironmentQuery/Generators/EnvQueryGenerator_SimpleGrid.cpp
--- a/Engine/Source/Runtime/Engine/Private/AI/EnvironmentQuery/Generators/EnvQueryGenerator_SimpleGrid.cpp
+++ b/Engine/Source/Runtime/Engine/Private/AI/EnvironmentQuery/Generators/EnvQueryGenerator_SimpleGrid.cpp
@@ -34,7 +34,7 @@ void UEnvQueryGenerator_SimpleGrid::GenerateItems(struct FEnvQueryInstance& Quer
 	}
 #endif // WITH_RECAST
 
-	const int32 ItemCount = FPlatformMath::Trunc((RadiusValue * 2.0f / DensityValue) + 1);
+	const int32 ItemCount = FPlatformMath::TruncToInt((RadiusValue * 2.0f / DensityValue) + 1);
 	const int32 ItemCountHalf = ItemCount / 2;
 
 	TArray<FVector> ContextLocations;
diff --git a/Engine/Source/Runtime/Engine/Private/AI/Navigation/PImplRecastNavMesh.cpp b/Engine/Source/Runtime/Engine/Private/AI/Navigation/PImplRecastNavMesh.cpp
--- a/Engine/Source/Runtime/Engine/Private/AI/Navigation/PImplRecastNavMesh.cpp
+++ b/Engine/Source/Runtime/Engine/Private/AI/Navigation/PImplRecastNavMesh.cpp
@@ -1762,7 +1762,7 @@ void FPImplRecastNavMesh::GetDebugTileBounds(FBox& OuterBox, int32& NumTilesX, i
  	if (DetourNavMesh)
  	{
  		dtNavMeshParams const* Params = DetourNavMesh->getParams();
-		NumTilesX = FMath::TruncFloat( FMath::Sqrt((float)Params->maxTiles) );
+		NumTilesX = FMath::TruncToFloat( FMath::Sqrt((float)Params->maxTiles) );
 		NumTilesY = NumTilesX;
 
  		FVector const Mn = Recast2UnrVector(Params->orig);
diff --git a/Engine/Source/Runtime/Engine/Private/AI/Navigation/RecastNavMesh.cpp b/Engine/Source/Runtime/Engine/Private/AI/Navigation/RecastNavMesh.cpp
--- a/Engine/Source/Runtime/Engine/Private/AI/Navigation/RecastNavMesh.cpp
+++ b/Engine/Source/Runtime/Engine/Private/AI/Navigation/RecastNavMesh.cpp
@@ -1770,7 +1770,7 @@ void ARecastNavMesh::PostEditChangeProperty( struct FPropertyChangedEvent& Prope
 				}
 
 				// tile's dimension can't exceed 2^16 x cell size, as it's being stored on 2 bytes
-				const int32 DimensionVX = FMath::Ceil(TileSizeUU / CellSize);
+				const int32 DimensionVX = FMath::CeilToInt(TileSizeUU / CellSize);
 				if (DimensionVX > MAX_uint16)
 				{
 					TileSizeUU = MAX_uint16 * CellSize;
diff --git a/Engine/Source/Runtime/Engine/Private/AI/Navigation/RecastNavMeshGenerator.cpp b/Engine/Source/Runtime/Engine/Private/AI/Navigation/RecastNavMeshGenerator.cpp
--- a/Engine/Source/Runtime/Engine/Private/AI/Navigation/RecastNavMeshGenerator.cpp
+++ b/Engine/Source/Runtime/Engine/Private/AI/Navigation/RecastNavMeshGenerator.cpp
@@ -1251,7 +1251,7 @@ struct FTileCacheCompressor : public dtTileCacheCompressor
 
 	virtual int32 maxCompressedSize(const int32 bufferSize)
 	{
-		return FMath::Trunc(bufferSize * 1.1f) + sizeof(FCompressedCacheHeader);
+		return FMath::TruncToInt(bufferSize * 1.1f) + sizeof(FCompressedCacheHeader);
 	}
 
 	virtual dtStatus compress(const uint8* buffer, const int32 bufferSize,
@@ -3007,7 +3007,7 @@ void FRecastNavMeshGenerator::Init()
 		const UNavigationSystem* NavSys = UNavigationSystem::GetCurrent(GetWorld());
 		Config.AgentIndex = NavSys ? NavSys->GetSupportedAgentIndex(NavGenParams) : 0;
 
-		Config.tileSize = FMath::Trunc(NavGenParams->TileSizeUU / CellSize);
+		Config.tileSize = FMath::TruncToInt(NavGenParams->TileSizeUU / CellSize);
 
 		Config.regionChunkSize = Config.tileSize / NavGenParams->LayerChunkSplits;
 		Config.TileCacheChunkSize = Config.tileSize / NavGenParams->RegionChunkSplits;
@@ -3049,7 +3049,7 @@ void FRecastNavMeshGenerator::Init()
 	// limit max amount of tiles: 64 bit poly address
 	const int32 MaxTileBits = 30;
 	const float AvgLayersPerTile = 8.0f;
-	const int32 MaxAllowedGridCells = FMath::Trunc((1 << MaxTileBits) / AvgLayersPerTile);
+	const int32 MaxAllowedGridCells = FMath::TruncToInt((1 << MaxTileBits) / AvgLayersPerTile);
 	const int32 NumRequestedCells = NewTilesWidth * NewTilesHeight;
 	if (NumRequestedCells < 0 || NumRequestedCells > MaxAllowedGridCells)
 	{
@@ -3126,7 +3126,7 @@ void FRecastNavMeshGenerator::Init()
 	// identifying a tile and a polygon.
 	int32 TileBits = MaxTileBits;
 	const float MaxLayers = TilesWidth * TilesHeight * AvgLayersPerTile;
-	TileBits = FMath::Min(FMath::Trunc(FMath::Log2(FMath::RoundUpToPowerOfTwo(MaxLayers))), MaxTileBits);
+	TileBits = FMath::Min(FMath::TruncToInt(FMath::Log2(FMath::RoundUpToPowerOfTwo(MaxLayers))), MaxTileBits);
 	MaxActiveTiles = 1 << TileBits;
 
 	const int32 PolyBits = FMath::Min(30, (int)(sizeof(dtPolyRef)*8 - DT_MIN_SALT_BITS) - TileBits);
@@ -3214,7 +3214,7 @@ void FRecastNavMeshGenerator::SetUpGeneration(float CellSize, float CellHeight,
 	Config.walkableSlopeAngle = AgentMaxSlope;
 	Config.walkableHeight = (int32)ceilf(AgentMinHeight / CellHeight);
 	Config.walkableClimb = (int32)ceilf(AgentMaxClimb / CellHeight);
-	const float WalkableRadius = FMath::Ceil(AgentRadius / CellSize);
+	const float WalkableRadius = FMath::CeilToFloat(AgentRadius / CellSize);
 	Config.walkableRadius = WalkableRadius;
 	
 	// store original sizes
@@ -3916,10 +3916,10 @@ void FRecastNavMeshGenerator::MarkDirtyGenerators()
 	for (int32 i = 0; i < DirtyAreasCopy.Num(); ++i, ++DirtyArea)
 	{
 		const FBox RCBB = Unreal2RecastBox(GrowBoundingBox(DirtyArea->Bounds));
-		const int32 XMin = FMath::Max(FMath::Trunc((RCBB.Min.X - RCNavBounds.Min.X) * InvTileCellSize), 0);
-		const int32 YMin = FMath::Max(FMath::Trunc((RCBB.Min.Z - RCNavBounds.Min.Z) * InvTileCellSize), 0);
-		const int32 XMax = FMath::Min(FMath::Trunc((RCBB.Max.X - RCNavBounds.Min.X) * InvTileCellSize), TilesWidth-1);
-		const int32 YMax = FMath::Min(FMath::Trunc((RCBB.Max.Z - RCNavBounds.Min.Z) * InvTileCellSize), TilesHeight-1);
+		const int32 XMin = FMath::Max(FMath::TruncToInt((RCBB.Min.X - RCNavBounds.Min.X) * InvTileCellSize), 0);
+		const int32 YMin = FMath::Max(FMath::TruncToInt((RCBB.Min.Z - RCNavBounds.Min.Z) * InvTileCellSize), 0);
+		const int32 XMax = FMath::Min(FMath::TruncToInt((RCBB.Max.X - RCNavBounds.Min.X) * InvTileCellSize), TilesWidth-1);
+		const int32 YMax = FMath::Min(FMath::TruncToInt((RCBB.Max.Z - RCNavBounds.Min.Z) * InvTileCellSize), TilesHeight-1);
 
 		for (int32 y = YMin; y <= YMax; ++y)
 		{
diff --git a/Engine/Source/Runtime/Engine/Private/Animation/AnimCompositeBase.cpp b/Engine/Source/Runtime/Engine/Private/Animation/AnimCompositeBase.cpp
--- a/Engine/Source/Runtime/Engine/Private/Animation/AnimCompositeBase.cpp
+++ b/Engine/Source/Runtime/Engine/Private/Animation/AnimCompositeBase.cpp
@@ -48,7 +48,7 @@ float FAnimSegment::ConvertTrackPosToAnimPos(const float & TrackPosition) const
 	const float AnimPositionUnWrapped = (TrackPosition - StartPos) * GetValidPlayRate();
 
 	// Figure out how many times animation is allowed to be looped.
-	const float LoopCount = FMath::Min(FMath::Floor(AnimPositionUnWrapped / AnimLength), FMath::Max(LoopingCount-1, 0));
+	const float LoopCount = FMath::Min(FMath::FloorToInt(AnimPositionUnWrapped / AnimLength), FMath::Max(LoopingCount-1, 0));
 	// Position within AnimSequence
 	const float AnimPosition = AnimStartTime + AnimPositionUnWrapped - float(LoopCount) * AnimLength;
 	
diff --git a/Engine/Source/Runtime/Engine/Private/Animation/AnimCompress.cpp b/Engine/Source/Runtime/Engine/Private/Animation/AnimCompress.cpp
--- a/Engine/Source/Runtime/Engine/Private/Animation/AnimCompress.cpp
+++ b/Engine/Source/Runtime/Engine/Private/Animation/AnimCompress.cpp
@@ -253,7 +253,7 @@ void UAnimCompress::BitwiseCompressAnimationTracks(
 						// write the frame values for each key
 						float KeyTime= SrcTrans.Times[KeyIndex];
 						float FrameTime = KeyTime * FrameRate;
-						int32 FrameIndex= FMath::Clamp(FMath::Trunc(FrameTime), 0, LastFrame);
+						int32 FrameIndex= FMath::Clamp(FMath::TruncToInt(FrameTime), 0, LastFrame);
 						AC_UnalignedWriteToStream( &FrameIndex, FrameSize );
 					}
 
@@ -355,7 +355,7 @@ void UAnimCompress::BitwiseCompressAnimationTracks(
 						// write the frame values for each key
 						float KeyTime = SrcRot.Times[KeyIndex];
 						float FrameTime = KeyTime * FrameRate;
-						int32 FrameIndex = FMath::Clamp(FMath::Trunc(FrameTime), 0, LastFrame);
+						int32 FrameIndex = FMath::Clamp(FMath::TruncToInt(FrameTime), 0, LastFrame);
 						AC_UnalignedWriteToStream( &FrameIndex, FrameSize );
 					}
 
@@ -441,7 +441,7 @@ void UAnimCompress::BitwiseCompressAnimationTracks(
 							// write the frame values for each key
 							float KeyTime= SrcScale.Times[KeyIndex];
 							float FrameTime = KeyTime * FrameRate;
-							int32 FrameIndex= FMath::Clamp(FMath::Trunc(FrameTime), 0, LastFrame);
+							int32 FrameIndex= FMath::Clamp(FMath::TruncToInt(FrameTime), 0, LastFrame);
 							AC_UnalignedWriteToStream( &FrameIndex, FrameSize );
 						}
 
diff --git a/Engine/Source/Runtime/Engine/Private/Animation/AnimCompress_PerTrackCompression.cpp b/Engine/Source/Runtime/Engine/Private/Animation/AnimCompress_PerTrackCompression.cpp
--- a/Engine/Source/Runtime/Engine/Private/Animation/AnimCompress_PerTrackCompression.cpp
+++ b/Engine/Source/Runtime/Engine/Private/Animation/AnimCompress_PerTrackCompression.cpp
@@ -791,7 +791,7 @@ protected:
 		for (int32 KeyIndex = 0; KeyIndex < NumKeys; ++KeyIndex)
 		{
 			// Convert the frame time into a frame index and write it out
-			FrameIndexType FrameIndex = (FrameIndexType)FMath::Clamp(FMath::Trunc(Times[KeyIndex] * FramesPerSecond), 0, NumFrames - 1);
+			FrameIndexType FrameIndex = (FrameIndexType)FMath::Clamp(FMath::TruncToInt(Times[KeyIndex] * FramesPerSecond), 0, NumFrames - 1);
 			AppendBytes(&FrameIndex, sizeof(FrameIndexType));
 		}
 
diff --git a/Engine/Source/Runtime/Engine/Private/Animation/AnimSequence.cpp b/Engine/Source/Runtime/Engine/Private/Animation/AnimSequence.cpp
--- a/Engine/Source/Runtime/Engine/Private/Animation/AnimSequence.cpp
+++ b/Engine/Source/Runtime/Engine/Private/Animation/AnimSequence.cpp
@@ -995,7 +995,7 @@ bool UAnimSequence::CropRawAnimData( float CurrentTime, bool bFromStart )
 	// The reason I'm changing to TotalNumOfFrames is CT/SL = KeyIndexWithFraction/TotalNumOfFrames
 	// To play TotalNumOfFrames, it takes SequenceLength. Each key will take SequenceLength/TotalNumOfFrames
 	float const KeyIndexWithFraction = (CurrentTime * (float)(TotalNumOfFrames)) / SequenceLength;
-	int32 KeyIndex = bFromStart ? FMath::Floor(KeyIndexWithFraction) : FMath::Ceil(KeyIndexWithFraction);
+	int32 KeyIndex = bFromStart ? FMath::FloorToInt(KeyIndexWithFraction) : FMath::CeilToInt(KeyIndexWithFraction);
 	// Ensure KeyIndex is in range.
 	KeyIndex = FMath::Clamp<int32>(KeyIndex, 1, TotalNumOfFrames-1); 
 	// determine which keys need to be removed.
diff --git a/Engine/Source/Runtime/Engine/Private/Animation/AnimationRuntime.cpp b/Engine/Source/Runtime/Engine/Private/Animation/AnimationRuntime.cpp
--- a/Engine/Source/Runtime/Engine/Private/Animation/AnimationRuntime.cpp
+++ b/Engine/Source/Runtime/Engine/Private/Animation/AnimationRuntime.cpp
@@ -643,7 +643,7 @@ void FAnimationRuntime::GetKeyIndicesFromTime(int32 & OutKeyIndex1, int32 & OutK
 	const float KeyPos = ((float)NumKeys * Time) / SequenceLength;
 
 	// Find the integer part (ensuring within range) and that gives us the 'starting' key index.
-	const int32 KeyIndex1 = FMath::Clamp<int32>( FMath::Floor(KeyPos), 0, NumFrames-1 );  // @todo should be changed to FMath::Trunc
+	const int32 KeyIndex1 = FMath::Clamp<int32>( FMath::FloorToInt(KeyPos), 0, NumFrames-1 );  // @todo should be changed to FMath::TruncToInt
 
 	// The alpha (fractional part) is then just the remainder.
 	const float Alpha = KeyPos - (float)KeyIndex1;
diff --git a/Engine/Source/Runtime/Engine/Private/Animation/BlendSpace.cpp b/Engine/Source/Runtime/Engine/Private/Animation/BlendSpace.cpp
--- a/Engine/Source/Runtime/Engine/Private/Animation/BlendSpace.cpp
+++ b/Engine/Source/Runtime/Engine/Private/Animation/BlendSpace.cpp
@@ -18,8 +18,8 @@ void UBlendSpace::GetGridSamplesFromBlendInput(const FVector &BlendInput, FGridB
 	FVector NormalizedBlendInput = GetNormalizedBlendInput(BlendInput);
 	
 	FVector GridIndex;
-	GridIndex.X = FMath::Trunc(NormalizedBlendInput.X);
-	GridIndex.Y = FMath::Trunc(NormalizedBlendInput.Y);
+	GridIndex.X = FMath::TruncToFloat(NormalizedBlendInput.X);
+	GridIndex.Y = FMath::TruncToFloat(NormalizedBlendInput.Y);
 	GridIndex.Z = 0.f;
 
 	FVector Remainder = NormalizedBlendInput - GridIndex;
diff --git a/Engine/Source/Runtime/Engine/Private/Animation/BlendSpace1D.cpp b/Engine/Source/Runtime/Engine/Private/Animation/BlendSpace1D.cpp
--- a/Engine/Source/Runtime/Engine/Private/Animation/BlendSpace1D.cpp
+++ b/Engine/Source/Runtime/Engine/Private/Animation/BlendSpace1D.cpp
@@ -66,7 +66,7 @@ void UBlendSpace1D::GetRawSamplesFromBlendInput(const FVector &BlendInput, TArra
 
 	FVector NormalizedBlendInput = GetNormalizedBlendInput(BlendInput);
 
-	float GridIndex = FMath::Trunc(NormalizedBlendInput.X);
+	float GridIndex = FMath::TruncToFloat(NormalizedBlendInput.X);
 	float Remainder = NormalizedBlendInput.X - GridIndex;
 
 	const FEditorElement* BeforeElement = GetGridSampleInternal(GridIndex);
diff --git a/Engine/Source/Runtime/Engine/Private/ContentStreaming.cpp b/Engine/Source/Runtime/Engine/Private/ContentStreaming.cpp
--- a/Engine/Source/Runtime/Engine/Private/ContentStreaming.cpp
+++ b/Engine/Source/Runtime/Engine/Private/ContentStreaming.cpp
@@ -499,7 +499,7 @@ FFloatMipLevel FFloatMipLevel::FromScreenSizeInTexels(float ScreenSizeInTexels)
 	FFloatMipLevel Ret;
 
 	// int
-//	Ret.MipLevel = 1 + FMath::CeilLogTwo( FMath::Trunc(ScreenSizeInTexels) );
+//	Ret.MipLevel = 1 + FMath::CeilLogTwo( FMath::TruncToInt(ScreenSizeInTexels) );
 	// float
 	Ret.MipLevel = 1.0f + FMath::Log2(ScreenSizeInTexels);
 
@@ -3727,7 +3727,7 @@ void FStreamingManagerTexture::CheckUserSettings()
 				float PoolSize = float(GPoolSizeVRAMPercentage) * 0.01f * float(TotalGraphicsMemory);
 
 				// Truncate the pool size to MB (but still count in bytes)
-				TexturePoolSize = int64(FGenericPlatformMath::TruncFloat(PoolSize / 1024.0f / 1024.0f)) * 1024 * 1024;
+				TexturePoolSize = int64(FGenericPlatformMath::TruncToFloat(PoolSize / 1024.0f / 1024.0f)) * 1024 * 1024;
 			}
 			else
 			{
@@ -5348,7 +5348,7 @@ FFloatMipLevel FStreamingHandlerTextureStatic::GetWantedMips( FStreamingManagerT
 
 								if (MinLODValue <= 0)
 								{
-									WantedMipCount = FMath::Max(WantedMipCount, FFloatMipLevel::FromMipLevel(StreamingTexture.MaxAllowedMips - 13 - FMath::Floor(MinLODValue)));
+									WantedMipCount = FMath::Max(WantedMipCount, FFloatMipLevel::FromMipLevel(StreamingTexture.MaxAllowedMips - 13 - FMath::FloorToInt(MinLODValue)));
 									if (WantedMipCount >= FFloatMipLevel::FromMipLevel(StreamingTexture.MaxAllowedMips))
 									{
 										MinDistance = 1.0f;
diff --git a/Engine/Source/Runtime/Engine/Private/Debug/GameplayDebuggingComponent.cpp b/Engine/Source/Runtime/Engine/Private/Debug/GameplayDebuggingComponent.cpp
--- a/Engine/Source/Runtime/Engine/Private/Debug/GameplayDebuggingComponent.cpp
+++ b/Engine/Source/Runtime/Engine/Private/Debug/GameplayDebuggingComponent.cpp
@@ -721,7 +721,7 @@ void UGameplayDebuggingComponent::CollectEQSData()
 
 	const int32 UncompressedSize = UncompressedBuffer.Num();
 	const int32 HeaderSize = sizeof(int32);
-	EQSRepData.Init(0, HeaderSize + FMath::Trunc(1.1f * UncompressedSize));
+	EQSRepData.Init(0, HeaderSize + FMath::TruncToInt(1.1f * UncompressedSize));
 
 	int32 CompressedSize = EQSRepData.Num() - HeaderSize;
 	uint8* DestBuffer = EQSRepData.GetTypedData();
@@ -736,7 +736,7 @@ void UGameplayDebuggingComponent::CollectEQSData()
 	UE_LOG(LogEQS, Log, TEXT("Preparing EQS data: %.1fkB took %.3fms (collect: %.3fms + compress %d%%: %.3fms)"),
 		EQSRepData.Num() / 1024.0f, 1000.0f * (Timer3 - Timer1),
 		1000.0f * (Timer2 - Timer1),
-		FMath::Trunc(100.0f * EQSRepData.Num() / UncompressedBuffer.Num()), 1000.0f * (Timer3 - Timer2));
+		FMath::TruncToInt(100.0f * EQSRepData.Num() / UncompressedBuffer.Num()), 1000.0f * (Timer3 - Timer2));
 
 	if (GetWorld()->GetNetMode() != NM_DedicatedServer)
 	{
@@ -1005,7 +1005,7 @@ void UGameplayDebuggingComponent::ServerCollectNavmeshData_Implementation(FVecto
 	const double Timer2 = FPlatformTime::Seconds();
 
 	const int32 HeaderSize = sizeof(int32);
-	NavmeshRepData.Init(0, HeaderSize + FMath::Trunc(1.1f * UncompressedBuffer.Num()));
+	NavmeshRepData.Init(0, HeaderSize + FMath::TruncToInt(1.1f * UncompressedBuffer.Num()));
 
 	const int32 UncompressedSize = UncompressedBuffer.Num();
 	int32 CompressedSize = NavmeshRepData.Num() - HeaderSize;
@@ -1022,7 +1022,7 @@ void UGameplayDebuggingComponent::ServerCollectNavmeshData_Implementation(FVecto
 	UE_LOG(LogNavigation, Log, TEXT("Preparing navmesh data: %.1fkB took %.3fms (collect: %.3fms + compress %d%%: %.3fms)"),
 		NavmeshRepData.Num() / 1024.0f, 1000.0f * (Timer3 - Timer1),
 		1000.0f * (Timer2 - Timer1),
-		FMath::Trunc(100.0f * NavmeshRepData.Num() / UncompressedBuffer.Num()), 1000.0f * (Timer3 - Timer2));
+		FMath::TruncToInt(100.0f * NavmeshRepData.Num() / UncompressedBuffer.Num()), 1000.0f * (Timer3 - Timer2));
 #endif
 
 	if (GetWorld()->GetNetMode() != NM_DedicatedServer)
diff --git a/Engine/Source/Runtime/Engine/Private/Distributions.cpp b/Engine/Source/Runtime/Engine/Private/Distributions.cpp
--- a/Engine/Source/Runtime/Engine/Private/Distributions.cpp
+++ b/Engine/Source/Runtime/Engine/Private/Distributions.cpp
@@ -1331,10 +1331,10 @@ void FComposableDistribution::QuantizeVector4(
 	Values = Table.Values.GetTypedData();
 	for ( int32 EntryIndex = 0; EntryIndex < EntryCount; ++EntryIndex )
 	{
-		QuantizedValues->R = FMath::Clamp<int32>( FMath::Trunc( (Values[0] - Bias.X) * InvScale.X ), 0, 255 );
-		QuantizedValues->G = FMath::Clamp<int32>( FMath::Trunc( (Values[1] - Bias.Y) * InvScale.Y ), 0, 255 );
-		QuantizedValues->B = FMath::Clamp<int32>( FMath::Trunc( (Values[2] - Bias.Z) * InvScale.Z ), 0, 255 );
-		QuantizedValues->A = FMath::Clamp<int32>( FMath::Trunc( (Values[3] - Bias.W) * InvScale.W ), 0, 255 );
+		QuantizedValues->R = FMath::Clamp<int32>( FMath::TruncToInt( (Values[0] - Bias.X) * InvScale.X ), 0, 255 );
+		QuantizedValues->G = FMath::Clamp<int32>( FMath::TruncToInt( (Values[1] - Bias.Y) * InvScale.Y ), 0, 255 );
+		QuantizedValues->B = FMath::Clamp<int32>( FMath::TruncToInt( (Values[2] - Bias.Z) * InvScale.Z ), 0, 255 );
+		QuantizedValues->A = FMath::Clamp<int32>( FMath::TruncToInt( (Values[3] - Bias.W) * InvScale.W ), 0, 255 );
 		Values += EntryStride;
 		QuantizedValues++;
 	}
diff --git a/Engine/Source/Runtime/Engine/Private/EdGraph/EdGraph.cpp b/Engine/Source/Runtime/Engine/Private/EdGraph/EdGraph.cpp
--- a/Engine/Source/Runtime/Engine/Private/EdGraph/EdGraph.cpp
+++ b/Engine/Source/Runtime/Engine/Private/EdGraph/EdGraph.cpp
@@ -487,8 +487,8 @@ void UEdGraphNode::GetPinHoverText(const UEdGraphPin& Pin, FString& HoverTextOut
 
 void UEdGraphNode::SnapToGrid(float GridSnapSize)
 {
-	NodePosX = GridSnapSize * FMath::Round(NodePosX/GridSnapSize);
-	NodePosY = GridSnapSize * FMath::Round(NodePosY/GridSnapSize);
+	NodePosX = GridSnapSize * FMath::RoundToInt(NodePosX/GridSnapSize);
+	NodePosY = GridSnapSize * FMath::RoundToInt(NodePosY/GridSnapSize);
 }
 
 void UEdGraphNode::DestroyNode()
diff --git a/Engine/Source/Runtime/Engine/Private/GPUSort.cpp b/Engine/Source/Runtime/Engine/Private/GPUSort.cpp
--- a/Engine/Source/Runtime/Engine/Private/GPUSort.cpp
+++ b/Engine/Source/Runtime/Engine/Private/GPUSort.cpp
@@ -1030,7 +1030,7 @@ static bool TestGPUSort_RenderThread(EGPUSortTest TestToRun)
 	case GPU_SORT_TEST_RANDOM:
 		for ( int32 i = 0; i < 1000; ++i )
 		{
-			int32 TestSize = FMath::Trunc(FMath::SRand() * (float)(GPU_SORT_TEST_SIZE_MAX - GPU_SORT_TEST_SIZE_MIN)) + GPU_SORT_TEST_SIZE_MIN;
+			int32 TestSize = FMath::TruncToInt(FMath::SRand() * (float)(GPU_SORT_TEST_SIZE_MAX - GPU_SORT_TEST_SIZE_MIN)) + GPU_SORT_TEST_SIZE_MIN;
 			if (TestGPUSortForSize( (TestSize + 0xF) & 0xFFFFFFF0 ) == false)
 			{
 				return false;
diff --git a/Engine/Source/Runtime/Engine/Private/GameViewportClient.cpp b/Engine/Source/Runtime/Engine/Private/GameViewportClient.cpp
--- a/Engine/Source/Runtime/Engine/Private/GameViewportClient.cpp
+++ b/Engine/Source/Runtime/Engine/Private/GameViewportClient.cpp
@@ -54,7 +54,7 @@ public:
 		if ( GShowDebugSelectedLightmap && GDebugSelectedComponent && GDebugSelectedLightmap )
 		{
 			float Radius = GDebugSelectedComponent->Bounds.SphereRadius;
-			int32 Sides = FMath::Clamp<int32>( FMath::Trunc(Radius*Radius*4.0f*PI/(80.0f*80.0f)), 8, 200 );
+			int32 Sides = FMath::Clamp<int32>( FMath::TruncToInt(Radius*Radius*4.0f*PI/(80.0f*80.0f)), 8, 200 );
 			DrawWireSphere( PDI, GDebugSelectedComponent->Bounds.Origin, FColor(255,130,0), GDebugSelectedComponent->Bounds.SphereRadius, Sides, SDPG_Foreground );
 		}
 #endif
@@ -814,7 +814,7 @@ void UGameViewportClient::Draw(FViewport* InViewport, FCanvas* SceneCanvas)
 			MinY = FMath::Min<uint32>(UpscaledViewRect.Min.Y, MinY);
 			MaxX = FMath::Max<uint32>(UpscaledViewRect.Max.X, MaxX);
 			MaxY = FMath::Max<uint32>(UpscaledViewRect.Max.Y, MaxY);
-			TotalArea += FMath::Trunc(UpscaledViewRect.Width()) * FMath::Trunc(UpscaledViewRect.Height());
+			TotalArea += FMath::TruncToInt(UpscaledViewRect.Width()) * FMath::TruncToInt(UpscaledViewRect.Height());
 		}
 
 		// To draw black borders around the rendered image (prevents artifacts from post processing passes that read outside of the image e.g. PostProcessAA)
@@ -894,7 +894,7 @@ void UGameViewportClient::Draw(FViewport* InViewport, FCanvas* SceneCanvas)
 					if (View != NULL)
 					{
 						// rendering to directly to viewport target
-						FVector CanvasOrigin(FMath::Trunc(View->UnscaledViewRect.Min.X), FMath::Trunc(View->UnscaledViewRect.Min.Y), 0.f);
+						FVector CanvasOrigin(FMath::TruncToFloat(View->UnscaledViewRect.Min.X), FMath::TruncToInt(View->UnscaledViewRect.Min.Y), 0.f);
 												
 						CanvasObject->Init(View->UnscaledViewRect.Width(), View->UnscaledViewRect.Height(), View);
 
@@ -942,7 +942,7 @@ void UGameViewportClient::Draw(FViewport* InViewport, FCanvas* SceneCanvas)
 
 							uint32 SizeX = SceneCanvas->GetRenderTarget()->GetSizeXY().X;
 							uint32 SizeY = SceneCanvas->GetRenderTarget()->GetSizeXY().Y;
-							FIntRect SubtitleRegion(FMath::Trunc(SizeX * MinPos.X), FMath::Trunc(SizeY * MinPos.Y), FMath::Trunc(SizeX * MaxPos.X), FMath::Trunc(SizeY * MaxPos.Y));
+							FIntRect SubtitleRegion(FMath::TruncToInt(SizeX * MinPos.X), FMath::TruncToInt(SizeY * MinPos.Y), FMath::TruncToInt(SizeX * MaxPos.X), FMath::TruncToInt(SizeY * MaxPos.Y));
 							// We need a world to do this
 							FSubtitleManager::GetSubtitleManager()->DisplaySubtitles( SceneCanvas, SubtitleRegion, GetWorld()->GetAudioTimeSeconds() );
 						}
diff --git a/Engine/Source/Runtime/Engine/Private/HUD.cpp b/Engine/Source/Runtime/Engine/Private/HUD.cpp
--- a/Engine/Source/Runtime/Engine/Private/HUD.cpp
+++ b/Engine/Source/Runtime/Engine/Private/HUD.cpp
@@ -382,7 +382,7 @@ void AHUD::DrawDebugTextList()
 				TextItem.SetColor( DebugTextList[Idx].TextColor );
 				TextItem.Text = FText::FromString( DebugTextList[Idx].DebugText );
 				TextItem.Scale = FVector2D( DebugTextList[Idx].FontScale, DebugTextList[Idx].FontScale);
-				DebugCanvas->DrawItem( TextItem, FVector2D( FMath::Ceil(ScreenLoc.X), FMath::Ceil(ScreenLoc.Y) ) );
+				DebugCanvas->DrawItem( TextItem, FVector2D( FMath::CeilToFloat(ScreenLoc.X), FMath::CeilToFloat(ScreenLoc.Y) ) );
 			}
 
 			// do this at the end so even small durations get at least one frame
diff --git a/Engine/Source/Runtime/Engine/Private/ImageUtils.cpp b/Engine/Source/Runtime/Engine/Private/ImageUtils.cpp
--- a/Engine/Source/Runtime/Engine/Private/ImageUtils.cpp
+++ b/Engine/Source/Runtime/Engine/Private/ImageUtils.cpp
@@ -43,16 +43,16 @@ void FImageUtils::ImageResize(int32 SrcWidth, int32 SrcHeight, const TArray<FCol
 			float EndY = SrcY + StepSizeY;
 			
 			// Generate a rectangular region of pixels and then find the average color of the region.
-			int32 PosY = FMath::Trunc(SrcY+0.5f);
+			int32 PosY = FMath::TruncToInt(SrcY+0.5f);
 			PosY = FMath::Clamp<int32>(PosY, 0, (SrcHeight - 1));
 
-			int32 PosX = FMath::Trunc(SrcX+0.5f);
+			int32 PosX = FMath::TruncToInt(SrcX+0.5f);
 			PosX = FMath::Clamp<int32>(PosX, 0, (SrcWidth - 1));
 
-			int32 EndPosY = FMath::Trunc(EndY+0.5f);
+			int32 EndPosY = FMath::TruncToInt(EndY+0.5f);
 			EndPosY = FMath::Clamp<int32>(EndPosY, 0, (SrcHeight - 1));
 
-			int32 EndPosX = FMath::Trunc(EndX+0.5f);
+			int32 EndPosX = FMath::TruncToInt(EndX+0.5f);
 			EndPosX = FMath::Clamp<int32>(EndPosX, 0, (SrcWidth - 1));
 
 			FColor FinalColor;
@@ -90,9 +90,9 @@ void FImageUtils::ImageResize(int32 SrcWidth, int32 SrcHeight, const TArray<FCol
 					}
 				}
 				StepColor /= (float)PixelCount;
-				uint8 FinalR = FMath::Clamp(FMath::Trunc(StepColor.X), 0, 255);
-				uint8 FinalG = FMath::Clamp(FMath::Trunc(StepColor.Y), 0, 255);
-				uint8 FinalB = FMath::Clamp(FMath::Trunc(StepColor.Z), 0, 255);
+				uint8 FinalR = FMath::Clamp(FMath::TruncToInt(StepColor.X), 0, 255);
+				uint8 FinalG = FMath::Clamp(FMath::TruncToInt(StepColor.Y), 0, 255);
+				uint8 FinalB = FMath::Clamp(FMath::TruncToInt(StepColor.Z), 0, 255);
 				FinalColor = FColor(FinalR, FinalG, FinalB);
 			}
 
@@ -185,8 +185,8 @@ void FImageUtils::CropAndScaleImage( int32 SrcWidth, int32 SrcHeight, int32 Desi
 	}
 
 	// Store crop width and height as ints for convenience
-	int32 CropWidth = FMath::Floor(MaxWidth);
-	int32 CropHeight = FMath::Floor(MaxHeight);
+	int32 CropWidth = FMath::FloorToInt(MaxWidth);
+	int32 CropHeight = FMath::FloorToInt(MaxHeight);
 
 	// Array holding the cropped image
 	TArray<FColor> CroppedData;
diff --git a/Engine/Source/Runtime/Engine/Private/Interpolation.cpp b/Engine/Source/Runtime/Engine/Private/Interpolation.cpp
--- a/Engine/Source/Runtime/Engine/Private/Interpolation.cpp
+++ b/Engine/Source/Runtime/Engine/Private/Interpolation.cpp
@@ -7367,8 +7367,8 @@ void UInterpTrackAnimControl::UpdateTrack(float NewPosition, UInterpTrackInst* T
 						// @todo: This will need to be updated if we decide to support notifies in reverse.
 						if(AnimSeq.bReverse == false)
 						{
-							int32 FromLoopNum = FMath::Floor( (((FromTime - CurrentSeqStart) * CurrentRate) + CurrentStartOffset)/SeqLength );
-							int32 ToLoopNum = FMath::Floor( (((ToTime - CurrentSeqStart) * CurrentRate) + CurrentStartOffset)/SeqLength );
+							int32 FromLoopNum = FMath::FloorToInt( (((FromTime - CurrentSeqStart) * CurrentRate) + CurrentStartOffset)/SeqLength );
+							int32 ToLoopNum = FMath::FloorToInt( (((ToTime - CurrentSeqStart) * CurrentRate) + CurrentStartOffset)/SeqLength );
 							int32 NumLoopsToJump = ToLoopNum - FromLoopNum;
 
 							for(int32 i=0; i<NumLoopsToJump; i++)
@@ -9405,7 +9405,7 @@ void UInterpTrackParticleReplay::UpdateTrack( float NewPosition, UInterpTrackIns
 				if( bHaveReplayStartKey )
 				{
 					const float TimeWithinReplay = NewPosition - CurrentReplayStartKey.Time;
-					const int32 ReplayFrameIndex = FMath::Trunc( TimeWithinReplay / FMath::Max( ( float )KINDA_SMALL_NUMBER, FixedTimeStep ) );
+					const int32 ReplayFrameIndex = FMath::TruncToInt( TimeWithinReplay / FMath::Max( ( float )KINDA_SMALL_NUMBER, FixedTimeStep ) );
 
 
 					// Check to see if we have a clip
