Fix for build break caused by CL#2070165.

diff --git a/Engine/Source/Programs/CrashReporter/CrashReportClient/Private/CrashReportUtil.h b/Engine/Source/Programs/CrashReporter/CrashReportClient/Private/CrashReportUtil.h
--- a/Engine/Source/Programs/CrashReporter/CrashReportClient/Private/CrashReportUtil.h
+++ b/Engine/Source/Programs/CrashReporter/CrashReportClient/Private/CrashReportUtil.h
@@ -31,9 +31,7 @@ public:
 	{
 		bIsDone = 0;
 
-		bool bAutoDeleteThread = true;
-		bool bAutoDeleteRunnable = false;
-		FRunnableThread::Create(this, Worker.Name(), bAutoDeleteThread, bAutoDeleteRunnable);
+		Thread = FRunnableThread::Create(this, Worker.Name());
 	}
 
 	/**
@@ -57,8 +55,17 @@ private:
 		return 0;
 	}
 
-	TWorker Worker;
+	virtual void Exit() OVERRIDE
+	{
+		if (Thread)
+		{
+			Thread->Kill();
+			delete Thread;
+		}
+	}
 
+	TWorker Worker;
+	FRunnableThread * Thread;
 	volatile int32 bIsDone;	
 };
 
diff --git a/Engine/Source/Programs/UnrealLightmass/Private/Lighting/LightingSystem.cpp b/Engine/Source/Programs/UnrealLightmass/Private/Lighting/LightingSystem.cpp
--- a/Engine/Source/Programs/UnrealLightmass/Private/Lighting/LightingSystem.cpp
+++ b/Engine/Source/Programs/UnrealLightmass/Private/Lighting/LightingSystem.cpp
@@ -558,7 +558,7 @@ void FStaticLightingSystem::MultithreadProcess()
 	{
 		FMappingProcessingThreadRunnable* ThreadRunnable = new(Threads) FMappingProcessingThreadRunnable(this, ThreadIndex, StaticLightingTask_ProcessMappings);
 		const FString ThreadName = FString::Printf(TEXT("MappingProcessingThread%u"), ThreadIndex);
-		ThreadRunnable->Thread = FRunnableThread::Create(ThreadRunnable, *ThreadName, 0, 0, 0, TPri_Normal);
+		ThreadRunnable->Thread = FRunnableThread::Create(ThreadRunnable, *ThreadName);
 	}
 	GStatistics.NumThreads = NumStaticLightingThreads + 1;	// Includes the main-thread who is only exporting.
 
diff --git a/Engine/Source/Programs/UnrealLightmass/Private/Lighting/PhotonMapping.cpp b/Engine/Source/Programs/UnrealLightmass/Private/Lighting/PhotonMapping.cpp
--- a/Engine/Source/Programs/UnrealLightmass/Private/Lighting/PhotonMapping.cpp
+++ b/Engine/Source/Programs/UnrealLightmass/Private/Lighting/PhotonMapping.cpp
@@ -257,7 +257,7 @@ void FStaticLightingSystem::EmitDirectPhotons(
 	{
 		FDirectPhotonEmittingThreadRunnable* ThreadRunnable = new(DirectPhotonEmittingThreads) FDirectPhotonEmittingThreadRunnable(this, ThreadIndex, Input);
 		const FString ThreadName = FString::Printf(TEXT("DirectPhotonEmittingThread%u"), ThreadIndex);
-		ThreadRunnable->Thread = FRunnableThread::Create(ThreadRunnable, *ThreadName, 0, 0, 0, TPri_Normal);
+		ThreadRunnable->Thread = FRunnableThread::Create(ThreadRunnable, *ThreadName);
 	}
 
 	const double StartEmittingDirectPhotonsMainThread = FPlatformTime::Seconds();
@@ -692,7 +692,7 @@ void FStaticLightingSystem::EmitIndirectPhotons(
 	{
 		FIndirectPhotonEmittingThreadRunnable* ThreadRunnable = new(IndirectPhotonEmittingThreads) FIndirectPhotonEmittingThreadRunnable(this, ThreadIndex, Input);
 		const FString ThreadName = FString::Printf(TEXT("IndirectPhotonEmittingThread%u"), ThreadIndex);
-		ThreadRunnable->Thread = FRunnableThread::Create(ThreadRunnable, *ThreadName, 0, 0, 0, TPri_Normal);
+		ThreadRunnable->Thread = FRunnableThread::Create(ThreadRunnable, *ThreadName);
 	}
 
 	const double StartEmittingIndirectPhotonsMainThread = FPlatformTime::Seconds();
@@ -1135,7 +1135,7 @@ void FStaticLightingSystem::MarkIrradiancePhotons(const FBoxSphereBounds& Import
 	{
 		FIrradiancePhotonMarkingThreadRunnable* ThreadRunnable = new(IrradiancePhotonMarkingThreads) FIrradiancePhotonMarkingThreadRunnable(this, ThreadIndex, IrradiancePhotons);
 		const FString ThreadName = FString::Printf(TEXT("IrradiancePhotonMarkingThread%u"), ThreadIndex);
-		ThreadRunnable->Thread = FRunnableThread::Create(ThreadRunnable, *ThreadName, 0, 0, 0, TPri_Normal);
+		ThreadRunnable->Thread = FRunnableThread::Create(ThreadRunnable, *ThreadName);
 	}
 
 	const double MainThreadStartTime = FPlatformTime::Seconds();
@@ -1299,7 +1299,7 @@ void FStaticLightingSystem::CalculateIrradiancePhotons(const FBoxSphereBounds& I
 	{
 		FIrradiancePhotonCalculatingThreadRunnable* ThreadRunnable = new(IrradiancePhotonThreads) FIrradiancePhotonCalculatingThreadRunnable(this, ThreadIndex, IrradiancePhotons);
 		const FString ThreadName = FString::Printf(TEXT("IrradiancePhotonCalculatingThread%u"), ThreadIndex);
-		ThreadRunnable->Thread = FRunnableThread::Create(ThreadRunnable, *ThreadName, 0, 0, 0, TPri_Normal);
+		ThreadRunnable->Thread = FRunnableThread::Create(ThreadRunnable, *ThreadName);
 	}
 
 	const double MainThreadStartTime = FPlatformTime::Seconds();
@@ -1498,7 +1498,7 @@ void FStaticLightingSystem::CacheIrradiancePhotons()
 	{
 		FMappingProcessingThreadRunnable* ThreadRunnable = new(IrradiancePhotonCachingThreads) FMappingProcessingThreadRunnable(this, ThreadIndex, StaticLightingTask_CacheIrradiancePhotons);
 		const FString ThreadName = FString::Printf(TEXT("IrradiancePhotonCachingThread%u"), ThreadIndex);
-		ThreadRunnable->Thread = FRunnableThread::Create(ThreadRunnable, *ThreadName, 0, 0, 0, TPri_Normal);
+		ThreadRunnable->Thread = FRunnableThread::Create(ThreadRunnable, *ThreadName);
 	}
 
 	// Start the static lighting thread loop on the main thread, too.
