Static and Stationary sky lights now have one bounce of diffuse GI computed by Lightmass.  This is done by caching a low quality version of the sky lighting on surfaces for the final gather samples to pick up, just like is done for normal direct lighting, although a lot more shadow rays are needed for sky lighting.

diff --git a/Engine/Source/Editor/UnrealEd/Private/Lightmass/Lightmass.cpp b/Engine/Source/Editor/UnrealEd/Private/Lightmass/Lightmass.cpp
--- a/Engine/Source/Editor/UnrealEd/Private/Lightmass/Lightmass.cpp
+++ b/Engine/Source/Editor/UnrealEd/Private/Lightmass/Lightmass.cpp
@@ -133,6 +133,7 @@ void Copy( const ULightComponentBase* In, Lightmass::FLightData& Out )
 	// Set brightness here for light types that only derive from ULightComponentBase and not from ULightComponent
 	Out.Brightness = In->Intensity;
 	Out.Guid = In->LightGuid;
+	Out.IndirectLightingScale = In->IndirectLightingIntensity;
 }
 
 void Copy( const ULightComponent* In, Lightmass::FLightData& Out )
@@ -891,7 +892,6 @@ void FLightmassExporter::WriteLights( int32 Channel )
 		Lightmass::FLightData LightData;
 		Lightmass::FDirectionalLightData DirectionalData;
 		Copy( Light, LightData );
-		LightData.IndirectLightingScale = Light->IndirectLightingIntensity;
 		LightData.IndirectLightingSaturation = Light->LightmassSettings.IndirectLightingSaturation;
 		LightData.ShadowExponent = Light->LightmassSettings.ShadowExponent;
 		LightData.LightSourceRadius = 0;
@@ -908,7 +908,6 @@ void FLightmassExporter::WriteLights( int32 Channel )
 		Lightmass::FLightData LightData;
 		Lightmass::FPointLightData PointData;
 		Copy( Light, LightData );
-		LightData.IndirectLightingScale = Light->IndirectLightingIntensity;
 		LightData.IndirectLightingSaturation = Light->LightmassSettings.IndirectLightingSaturation;
 		LightData.ShadowExponent = Light->LightmassSettings.ShadowExponent;
 		LightData.LightSourceRadius = Light->SourceRadius;
@@ -927,7 +926,6 @@ void FLightmassExporter::WriteLights( int32 Channel )
 		Lightmass::FPointLightData PointData;
 		Lightmass::FSpotLightData SpotData;
 		Copy( Light, LightData ); 
-		LightData.IndirectLightingScale = Light->IndirectLightingIntensity;
 		LightData.IndirectLightingSaturation = Light->LightmassSettings.IndirectLightingSaturation;
 		LightData.ShadowExponent = Light->LightmassSettings.ShadowExponent;
 		LightData.LightSourceRadius = Light->SourceRadius;
diff --git a/Engine/Source/Programs/UnrealLightmass/Private/Lighting/FinalGather.cpp b/Engine/Source/Programs/UnrealLightmass/Private/Lighting/FinalGather.cpp
--- a/Engine/Source/Programs/UnrealLightmass/Private/Lighting/FinalGather.cpp
+++ b/Engine/Source/Programs/UnrealLightmass/Private/Lighting/FinalGather.cpp
@@ -156,7 +156,7 @@ FLinearColor FStaticLightingSystem::EvaluateEnvironmentLighting(const FVector4&
 	return IncomingDirection.Z < 0 ? (MaterialSettings.EnvironmentColor / (float)PI) : FLinearColor::Black;
 }
 
-void FStaticLightingSystem::EvaluateSkyLighting(const FVector4& IncomingDirection, bool bShadowed, FLinearColor& OutStaticLighting, FLinearColor& OutStationaryLighting) const
+void FStaticLightingSystem::EvaluateSkyLighting(const FVector4& IncomingDirection, bool bShadowed, bool bForDirectLighting, FLinearColor& OutStaticLighting, FLinearColor& OutStationaryLighting) const
 {
 	for (int32 LightIndex = 0; LightIndex < SkyLights.Num(); LightIndex++)
 	{
@@ -165,7 +165,8 @@ void FStaticLightingSystem::EvaluateSkyLighting(const FVector4& IncomingDirectio
 		if (!bShadowed || !(SkyLight->LightFlags & GI_LIGHT_CASTSHADOWS))
 		{
 			FSHVector3 SH = FSHVector3::SHBasisFunction(IncomingDirection);
-			FLinearColor Lighting = Dot(SkyLight->IrradianceEnvironmentMap, SH) * SkyLight->Brightness * FLinearColor(SkyLight->Color);
+			const float LightingScale = bForDirectLighting ? 1.0f : SkyLight->IndirectLightingScale;
+			FLinearColor Lighting = (Dot(SkyLight->IrradianceEnvironmentMap, SH) * SkyLight->Brightness * LightingScale) * FLinearColor(SkyLight->Color);
 
 			Lighting.R = FMath::Max(Lighting.R, 0.0f);
 			Lighting.G = FMath::Max(Lighting.G, 0.0f);
@@ -558,7 +559,7 @@ FLinearColor FStaticLightingSystem::FinalGatherSample(
 		}
 	}
 
-	EvaluateSkyLighting(WorldPathDirection, RayIntersection.bIntersects, Lighting, OutStationarySkyLighting);
+	EvaluateSkyLighting(WorldPathDirection, RayIntersection.bIntersects, true, Lighting, OutStationarySkyLighting);
 
 #if ALLOW_LIGHTMAP_SAMPLE_DEBUGGING
 	if (bDebugThisTexel
@@ -1216,8 +1217,6 @@ SampleType FStaticLightingSystem::IncomingRadianceUniform(
 	//@todo - use cosine sampling if possible to match the indirect integrand, the irradiance caching algorithm assumes uniform sampling
 	for (int32 SampleIndex = 0; SampleIndex < UniformHemisphereSamples.Num(); SampleIndex++)
 	{
-		//const FVector4& SampleDirection = UniformHemisphereSamples[SampleIndex];
-		//const FVector4 TriangleTangentPathDirection = ((FVector)SampleDirection).RotateAngleAxis(RandomStream.GetFraction() * 360, FVector(0, 0, 1));
 		const FVector4 TriangleTangentPathDirection = UniformHemisphereSamples[SampleIndex];
 		checkSlow(TriangleTangentPathDirection.Z >= 0.0f);
 		checkSlow(TriangleTangentPathDirection.IsUnit3());
@@ -1372,7 +1371,7 @@ SampleType FStaticLightingSystem::IncomingRadianceUniform(
 
 		FLinearColor StaticSkyLighting = FLinearColor::Black;
 		FLinearColor StationarySkyLighting = FLinearColor::Black;
-		EvaluateSkyLighting(WorldPathDirection, RayIntersection.bIntersects, StaticSkyLighting, StationarySkyLighting);
+		EvaluateSkyLighting(WorldPathDirection, RayIntersection.bIntersects, true, StaticSkyLighting, StationarySkyLighting);
 
 		if (StaticSkyLighting != FLinearColor::Black)
 		{
diff --git a/Engine/Source/Programs/UnrealLightmass/Private/Lighting/LightingSystem.cpp b/Engine/Source/Programs/UnrealLightmass/Private/Lighting/LightingSystem.cpp
--- a/Engine/Source/Programs/UnrealLightmass/Private/Lighting/LightingSystem.cpp
+++ b/Engine/Source/Programs/UnrealLightmass/Private/Lighting/LightingSystem.cpp
@@ -1311,6 +1311,19 @@ void FStaticLightingSystem::CacheSamples()
 
 	CachedSamplesMaxUnoccludedLength = (CombinedVector / CachedHemisphereSamples.Num()).Size3();
 
+	{
+		int32 TargetNumApproximateSkyLightingSamples = FMath::Max(FMath::TruncToInt(ImportanceTracingSettings.NumHemisphereSamples / 5 * GeneralSettings.IndirectLightingQuality), 12);
+		CachedHemisphereSamplesForApproximateSkyLighting.Empty(TargetNumApproximateSkyLightingSamples);
+		TArray<FVector2D> Unused;
+		Unused.Empty(TargetNumApproximateSkyLightingSamples);
+
+		const float NumThetaStepsFloat = FMath::Sqrt(TargetNumApproximateSkyLightingSamples / (float)PI);
+		const int32 NumThetaSteps = FMath::TruncToInt(NumThetaStepsFloat);
+		const int32 NumPhiSteps = FMath::TruncToInt(NumThetaStepsFloat * (float)PI);
+
+		GenerateStratifiedUniformHemisphereSamples(NumThetaSteps, NumPhiSteps, RandomStream, CachedHemisphereSamplesForApproximateSkyLighting, Unused);
+	}
+
 	// Cache samples on the surface of each light for area shadows
 	for (int32 LightIndex = 0; LightIndex < Lights.Num(); LightIndex++)
 	{
@@ -2003,6 +2016,68 @@ void FStaticLightingSystem::CalculateApproximateDirectLighting(
 	}
 }
 
+FGatheredLightSample FStaticLightingSystem::CalculateApproximateSkyLighting(
+	const FFullStaticLightingVertex& Vertex,
+	float SampleRadius,
+	const TArray<FVector4>& UniformHemisphereSamples,
+	FStaticLightingMappingContext& MappingContext) const
+{
+	FGatheredLightSample IncomingRadiance;
+
+	if (SkyLights.Num() > 0)
+	{
+		const float UniformPDF = 1.0f / (2.0f * (float)PI);
+		const float SampleWeight = 1.0f / (UniformPDF * UniformHemisphereSamples.Num());
+
+		for (int32 SampleIndex = 0; SampleIndex < UniformHemisphereSamples.Num(); SampleIndex++)
+		{
+			const FVector4 TriangleTangentPathDirection = UniformHemisphereSamples[SampleIndex];
+			checkSlow(TriangleTangentPathDirection.Z >= 0.0f);
+			checkSlow(TriangleTangentPathDirection.IsUnit3());
+
+			// Generate the uniform hemisphere samples from a hemisphere based around the triangle normal, not the smoothed vertex normal
+			// This is important for cases where the smoothed vertex normal is very different from the triangle normal, in which case
+			// Using the smoothed vertex normal would cause self-intersection even on a plane
+			const FVector4 WorldPathDirection = Vertex.TransformTriangleTangentVectorToWorld(TriangleTangentPathDirection);
+			checkSlow(WorldPathDirection.IsUnit3());
+
+			const FVector4 TangentPathDirection = Vertex.TransformWorldVectorToTangent(WorldPathDirection);
+			checkSlow(TangentPathDirection.IsUnit3());
+
+			const FLightRay PathRay(
+				// Apply various offsets to the start of the ray.
+				// The offset along the ray direction is to avoid incorrect self-intersection due to floating point precision.
+				// The offset along the normal is to push self-intersection patterns (like triangle shape) on highly curved surfaces onto the backfaces.
+				Vertex.WorldPosition 
+				+ WorldPathDirection * SceneConstants.VisibilityRayOffsetDistance 
+				+ Vertex.WorldTangentZ * SampleRadius * SceneConstants.VisibilityNormalOffsetSampleRadiusScale,
+				Vertex.WorldPosition + WorldPathDirection * MaxRayDistance,
+				NULL,
+				NULL
+				);
+
+			FLightRayIntersection RayIntersection;
+			AggregateMesh.IntersectLightRay(PathRay, false, false, false, MappingContext.RayCache, RayIntersection);
+
+			FLinearColor StaticSkyLighting = FLinearColor::Black;
+			FLinearColor StationarySkyLighting = FLinearColor::Black;
+			EvaluateSkyLighting(WorldPathDirection, RayIntersection.bIntersects, false, StaticSkyLighting, StationarySkyLighting);
+
+			if (StaticSkyLighting != FLinearColor::Black)
+			{
+				IncomingRadiance.AddWeighted(FGatheredLightSample::PointLightWorldSpace(StaticSkyLighting, TangentPathDirection, WorldPathDirection), SampleWeight);
+			}
+
+			if (StationarySkyLighting != FLinearColor::Black)
+			{
+				IncomingRadiance.AddWeighted(FGatheredLightSample::PointLightWorldSpace(StationarySkyLighting, TangentPathDirection, WorldPathDirection), SampleWeight);
+			}
+		}
+	}
+
+	return IncomingRadiance;
+}
+
 /** Initializes DominantLightShadowInfo and prepares for multithreaded generation of DominantLightShadowInfo.ShadowMap. */
 void FStaticLightingSystem::BeginCalculateDominantShadowInfo(FGuid LightGuid)
 {
diff --git a/Engine/Source/Programs/UnrealLightmass/Private/Lighting/LightingSystem.h b/Engine/Source/Programs/UnrealLightmass/Private/Lighting/LightingSystem.h
--- a/Engine/Source/Programs/UnrealLightmass/Private/Lighting/LightingSystem.h
+++ b/Engine/Source/Programs/UnrealLightmass/Private/Lighting/LightingSystem.h
@@ -2175,6 +2175,12 @@ private:
 		FGatheredLightSample& OutToggleableDirectLighting,
 		float& OutToggleableDirectionalLightShadowing) const;
 
+	FGatheredLightSample CalculateApproximateSkyLighting(
+		const FFullStaticLightingVertex& Vertex,
+		float SampleRadius,
+		const TArray<FVector4>& UniformHemisphereSamples,
+		FStaticLightingMappingContext& MappingContext) const;
+
 	/** Calculates incident radiance for a given world space position. */
 	void CalculateVolumeSampleIncidentRadiance(
 		const TArray<FVector4>& UniformHemisphereSamples,
@@ -2251,7 +2257,7 @@ private:
 	FLinearColor EvaluateEnvironmentLighting(const FVector4& IncomingDirection) const;
 
 	/** Evaluates the incoming sky lighting from the scene. */
-	void EvaluateSkyLighting(const FVector4& IncomingDirection, bool bShadowed, FLinearColor& OutStaticLighting, FLinearColor& OutStationaryLighting) const;
+	void EvaluateSkyLighting(const FVector4& IncomingDirection, bool bShadowed, bool bForDirectLighting, FLinearColor& OutStaticLighting, FLinearColor& OutStationaryLighting) const;
 
 	/** Returns a light sample that represents the material attribute specified by MaterialSettings.ViewMaterialAttribute at the intersection. */
 	FGatheredLightSample GetVisualizedMaterialAttribute(const FStaticLightingMapping* Mapping, const FLightRayIntersection& Intersection) const;
@@ -2633,6 +2639,8 @@ private:
 
 	TArray<FVector2D> CachedHemisphereSampleUniforms;
 
+	TArray<FVector4> CachedHemisphereSamplesForApproximateSkyLighting;
+
 	/** The aggregate mesh used for raytracing. */
 	FStaticLightingAggregateMesh AggregateMesh;
 	
diff --git a/Engine/Source/Programs/UnrealLightmass/Private/Lighting/TextureMapping.cpp b/Engine/Source/Programs/UnrealLightmass/Private/Lighting/TextureMapping.cpp
--- a/Engine/Source/Programs/UnrealLightmass/Private/Lighting/TextureMapping.cpp
+++ b/Engine/Source/Programs/UnrealLightmass/Private/Lighting/TextureMapping.cpp
@@ -297,24 +297,31 @@ void FStaticLightingSystem::CacheIrradiancePhotonsTextureMapping(FStaticLighting
 		TextureMapping->Mesh->GetTriangle(TriangleIndex, V0.Vertex, V1.Vertex, V2.Vertex, Element);
 		V0.ElementIndex = V1.ElementIndex = V2.ElementIndex = Element;
 
-		// Rasterize the triangle using the mapping's texture coordinate channel.
-		FTriangleRasterizer<FStaticLightingRasterPolicy> TexelMappingRasterizer(FStaticLightingRasterPolicy(
-				Scene,
-				TexelToVertexMap,
-				SampleWeight,
-				FVector4(0),
-				false,
+		const FVector4 TriangleNormal = ((V2.Vertex.WorldPosition - V0.Vertex.WorldPosition) ^ (V1.Vertex.WorldPosition - V0.Vertex.WorldPosition)).SafeNormal();
+
+		// Don't rasterize degenerates 
+		if (!TriangleNormal.IsNearlyZero3())
+		{
+			// Rasterize the triangle using the mapping's texture coordinate channel.
+			FTriangleRasterizer<FStaticLightingRasterPolicy> TexelMappingRasterizer(FStaticLightingRasterPolicy(
+					Scene,
+					TexelToVertexMap,
+					SampleWeight,
+					TriangleNormal,
+					false,
+					false
+					));
+
+			TexelMappingRasterizer.DrawTriangle(
+				V0,
+				V1,
+				V2,
+				V0.Vertex.TextureCoordinates[TextureMapping->LightmapTextureCoordinateIndex] * FVector2D(TextureMapping->IrradiancePhotonCacheSizeX,TextureMapping->IrradiancePhotonCacheSizeY) + FVector2D(-0.5f,-0.5f),
+				V1.Vertex.TextureCoordinates[TextureMapping->LightmapTextureCoordinateIndex] * FVector2D(TextureMapping->IrradiancePhotonCacheSizeX,TextureMapping->IrradiancePhotonCacheSizeY) + FVector2D(-0.5f,-0.5f),
+				V2.Vertex.TextureCoordinates[TextureMapping->LightmapTextureCoordinateIndex] * FVector2D(TextureMapping->IrradiancePhotonCacheSizeX,TextureMapping->IrradiancePhotonCacheSizeY) + FVector2D(-0.5f,-0.5f),
 				false
-				));
-		TexelMappingRasterizer.DrawTriangle(
-			V0,
-			V1,
-			V2,
-			V0.Vertex.TextureCoordinates[TextureMapping->LightmapTextureCoordinateIndex] * FVector2D(TextureMapping->IrradiancePhotonCacheSizeX,TextureMapping->IrradiancePhotonCacheSizeY) + FVector2D(-0.5f,-0.5f),
-			V1.Vertex.TextureCoordinates[TextureMapping->LightmapTextureCoordinateIndex] * FVector2D(TextureMapping->IrradiancePhotonCacheSizeX,TextureMapping->IrradiancePhotonCacheSizeY) + FVector2D(-0.5f,-0.5f),
-			V2.Vertex.TextureCoordinates[TextureMapping->LightmapTextureCoordinateIndex] * FVector2D(TextureMapping->IrradiancePhotonCacheSizeX,TextureMapping->IrradiancePhotonCacheSizeY) + FVector2D(-0.5f,-0.5f),
-			false
-			);
+				);
+		}
 	}
 
 	// Allocate space for the cached irradiance photons
@@ -343,15 +350,38 @@ void FStaticLightingSystem::CacheIrradiancePhotonsTextureMapping(FStaticLighting
 			if (TexelToVertex.TotalSampleWeight > 0.0f)
 			{
 				MappingContext.Stats.NumCachedIrradianceSamples++;
-				const FStaticLightingVertex CurrentVertex = TexelToVertex.GetVertex();
-				
+				FFullStaticLightingVertex CurrentVertex = TexelToVertex.GetFullVertex();
+
+				// Normalize the tangent basis and ensure it is orthonormal
+				CurrentVertex.WorldTangentZ = TexelToVertex.WorldTangentZ.UnsafeNormal3();
+				CurrentVertex.TriangleNormal = TexelToVertex.TriangleNormal.UnsafeNormal3();
+				checkSlow(!CurrentVertex.TriangleNormal.ContainsNaN());
+
+				const FVector4 OriginalTangentX = CurrentVertex.WorldTangentX;
+				const FVector4 OriginalTangentY = CurrentVertex.WorldTangentY;
+
+				CurrentVertex.WorldTangentY = (CurrentVertex.WorldTangentZ ^ CurrentVertex.WorldTangentX).UnsafeNormal3();
+				// Maintain handedness
+				if (Dot3(CurrentVertex.WorldTangentY, OriginalTangentY) < 0)
+				{
+					CurrentVertex.WorldTangentY *= -1.0f;
+				}
+				CurrentVertex.WorldTangentX = CurrentVertex.WorldTangentY ^ CurrentVertex.WorldTangentZ;
+				if (Dot3(CurrentVertex.WorldTangentX, OriginalTangentX) < 0)
+				{
+					CurrentVertex.WorldTangentX *= -1.0f;
+				}
+
 				FGatheredLightSample DirectLighting;
 				FGatheredLightSample Unused;
 				float Unused2;
 
 				CalculateApproximateDirectLighting(CurrentVertex, TexelToVertex.TexelRadius, .1f, true, true, bDebugThisTexel && PhotonMappingSettings.bVisualizeCachedApproximateDirectLighting, MappingContext, DirectLighting, Unused, Unused2);
 
-				TextureMapping->CachedDirectLighting[Y * TextureMapping->IrradiancePhotonCacheSizeX + X] = DirectLighting.IncidentLighting;
+				// Compute low quality sky lighting and cache in the direct lighting, so we get one bounce, since sky lighting isn't handled by photons
+				const FGatheredLightSample SkyLighting = CalculateApproximateSkyLighting(CurrentVertex, TexelToVertex.TexelRadius, CachedHemisphereSamplesForApproximateSkyLighting, MappingContext);
+
+				TextureMapping->CachedDirectLighting[Y * TextureMapping->IrradiancePhotonCacheSizeX + X] = DirectLighting.IncidentLighting + SkyLighting.IncidentLighting;
 
 				// Only search the irradiance photon map if the surface cache position is inside the importance volume,
 				// Since irradiance photons are only deposited inside the importance volume.
diff --git a/Engine/Source/Runtime/Engine/Classes/Components/LightComponent.h b/Engine/Source/Runtime/Engine/Classes/Components/LightComponent.h
--- a/Engine/Source/Runtime/Engine/Classes/Components/LightComponent.h
+++ b/Engine/Source/Runtime/Engine/Classes/Components/LightComponent.h
@@ -19,13 +19,6 @@ class ENGINE_API ULightComponent : public ULightComponentBase
 	/** Transient shadowmap channel used to preview the results of stationary light shadowmap packing. */
 	int32 PreviewShadowMapChannel;
 	
-	/** 
-	 * Scales the indirect lighting contribution from this light. 
-	 * A value of 0 disables any GI from this light. Default is 1.
-	 */
-	UPROPERTY(BlueprintReadOnly, interp, Category=Light, meta=(UIMin = "0.0", UIMax = "6.0"))
-	float IndirectLightingIntensity;
-
 	/** Radius of light source shape. Moved to point light */
 	UPROPERTY()
 	float SourceRadius_DEPRECATED;
diff --git a/Engine/Source/Runtime/Engine/Classes/Components/LightComponentBase.h b/Engine/Source/Runtime/Engine/Classes/Components/LightComponentBase.h
--- a/Engine/Source/Runtime/Engine/Classes/Components/LightComponentBase.h
+++ b/Engine/Source/Runtime/Engine/Classes/Components/LightComponentBase.h
@@ -69,6 +69,13 @@ class ENGINE_API ULightComponentBase : public USceneComponent
 	UPROPERTY()
 	uint32 bPrecomputedLightingIsValid:1;
 
+	/** 
+	 * Scales the indirect lighting contribution from this light. 
+	 * A value of 0 disables any GI from this light. Default is 1.
+	 */
+	UPROPERTY(BlueprintReadOnly, interp, Category=Light, meta=(UIMin = "0.0", UIMax = "6.0"))
+	float IndirectLightingIntensity;
+
 #if WITH_EDITORONLY_DATA
 	/** Sprite for static light in the editor. */
 	UPROPERTY(transient)
diff --git a/Engine/Source/Runtime/Engine/Private/SkyLightComponent.cpp b/Engine/Source/Runtime/Engine/Private/SkyLightComponent.cpp
--- a/Engine/Source/Runtime/Engine/Private/SkyLightComponent.cpp
+++ b/Engine/Source/Runtime/Engine/Private/SkyLightComponent.cpp
@@ -96,6 +96,7 @@ USkyLightComponent::USkyLightComponent(const class FPostConstructInitializePrope
 
 	Brightness_DEPRECATED = 1;
 	Intensity = 1;
+	IndirectLightingIntensity = 1.0f;
 	SkyDistanceThreshold = 150000;
 	Mobility = EComponentMobility::Stationary;
 	bCaptureDirty = false;
