Static shadowing for translucency
* Each Stationary light gets a static shadow depth map generated by Lightmass, which is then used to shadow during translucent lighting volume injection

diff --git a/Engine/Config/BaseLightmass.ini b/Engine/Config/BaseLightmass.ini
--- a/Engine/Config/BaseLightmass.ini
+++ b/Engine/Config/BaseLightmass.ini
@@ -115,6 +115,11 @@ MaxTransitionDistanceWorldSpace=50
 ApproximateHighResTexelsPerMaxTransitionDistance=50
 MinDistanceFieldUpsampleFactor=3
 MinUnoccludedFraction=.005
+StaticShadowDepthMapTransitionSampleDistanceX=100
+StaticShadowDepthMapTransitionSampleDistanceY=100
+StaticShadowDepthMapSuperSampleFactor=2
+; Clamp the number of shadow samples generated to ~8mb for huge levels
+StaticShadowDepthMapMaxSamples=4194304
 
 [DevOptions.ImportanceTracing]
 bUseCosinePDF=False
diff --git a/Engine/Shaders/TranslucentLightInjectionShaders.usf b/Engine/Shaders/TranslucentLightInjectionShaders.usf
--- a/Engine/Shaders/TranslucentLightInjectionShaders.usf
+++ b/Engine/Shaders/TranslucentLightInjectionShaders.usf
@@ -103,6 +103,16 @@ float4 ShadowInjectParams;
 float SpotlightMask;
 float2 DepthBiasParameters;
 
+/** Whether to compute static shadowing. */
+uint bStaticallyShadowed;
+
+/** Shadow depth map computed for static geometry by Lightmass. */
+Texture2D StaticShadowDepthTexture;
+SamplerState StaticShadowDepthTextureSampler;
+
+/** Transform used for static shadowing by spot and directional lights. */
+float4x4 WorldToStaticShadowMatrix;
+
 float GetLightFunctionShadowFactor(float3 WorldPositionForLighting)
 {
 	float ShadowFactor = 1;
@@ -114,7 +124,97 @@ float GetLightFunctionShadowFactor(float3 WorldPositionForLighting)
 		float3 LightFunction = GetLightFunctionColor(LightVector.xyz, WorldPositionForLighting);
 
 		// We only suport monochrome light functions
-		ShadowFactor = dot(LightFunction, .3333f).xxx;
+		ShadowFactor = dot(LightFunction, .3333f).x;
+	#endif
+
+	return ShadowFactor;
+}
+
+/** Computes dynamic and static shadowing for a point anywhere in space. */
+float ComputeVolumeShadowing(float3 WorldPositionForLighting)
+{
+	float ShadowFactor = 1;
+
+	BRANCH
+	if (bStaticallyShadowed)
+	{
+		bool bUsePointLightShadowing = false;
+
+		#if RADIAL_ATTENUATION
+			bUsePointLightShadowing = SpotlightMask < 1;
+		#endif
+
+		BRANCH
+		if (bUsePointLightShadowing)
+		{
+			float3 LightVector = WorldPositionForLighting - DeferredLightUniforms.LightPosition;
+			float DistanceToLight = length(LightVector);
+			float3 NormalizedLightVector = LightVector / DistanceToLight;
+
+			//@todo - use parametrization without slow inverse trig.  Dual paraboloid?
+			float NormalizedTheta = atan2(NormalizedLightVector.y, NormalizedLightVector.x) / (2 * PI);
+			// atan2 returns in the range [-PI, PI], wrap the negative portion to [.5, 1]
+			float U = NormalizedTheta > 0 ? NormalizedTheta : 1 + NormalizedTheta;
+			float V = acos(NormalizedLightVector.z) / PI;
+			float2 UnwrappedUVs = float2(U, V);
+
+			float ShadowDepth = Texture2DSampleLevel(StaticShadowDepthTexture, StaticShadowDepthTextureSampler, UnwrappedUVs, 0).x;
+			ShadowFactor = DistanceToLight * DeferredLightUniforms.LightInvRadius < ShadowDepth;
+		}
+		else
+		{
+			// This path is used for directional lights and spot lights, which only require a single projection
+			// Transform the world position into shadowmap space
+			float4 HomogeneousShadowPosition = mul(float4(WorldPositionForLighting, 1), WorldToStaticShadowMatrix);
+			float2 ShadowUVs = HomogeneousShadowPosition.xy / HomogeneousShadowPosition.w;
+
+			// Treat as unshadowed if the voxel is outside of the shadow map
+			if (all(ShadowUVs >= 0 && ShadowUVs <= 1))
+			{
+				// Sample the shadowmap depth and determine if this voxel is shadowed
+				float ShadowDepth = Texture2DSampleLevel(StaticShadowDepthTexture, StaticShadowDepthTextureSampler, ShadowUVs, 0).x;
+				ShadowFactor = HomogeneousShadowPosition.z < ShadowDepth;
+			}
+		}
+	}		
+
+	#if DYNAMICALLY_SHADOWED
+		bool bUseCubemapShadowing = false;
+		float DynamicShadowFactor = 1;
+
+		#if RADIAL_ATTENUATION
+			bUseCubemapShadowing = SpotlightMask < 1;
+		#endif
+
+		if (bUseCubemapShadowing)
+		{
+			DynamicShadowFactor = CubemapHardwarePCF(WorldPositionForLighting, DeferredLightUniforms.LightPosition, DeferredLightUniforms.LightInvRadius, 0.03f * 512 * InvShadowmapResolution);
+		}
+		else
+		{
+			// Transform the world position into shadowmap space
+			float4 HomogeneousShadowPosition = mul(float4(WorldPositionForLighting, 1), WorldToShadowMatrix);
+			float2 ShadowUVs = HomogeneousShadowPosition.xy / HomogeneousShadowPosition.w;
+
+			// Treat as unshadowed if the voxel is outside of the shadow map
+			if (all(ShadowUVs >= ShadowmapMinMax.xy && ShadowUVs <= ShadowmapMinMax.zw))
+			{
+				// Sample the shadowmap depth and determine if this voxel is shadowed
+				float ShadowDepth = Texture2DSampleLevel(ShadowDepthTexture, ShadowDepthTextureSampler, ShadowUVs, 0).x;
+				DynamicShadowFactor = HomogeneousShadowPosition.z < ShadowDepth - DepthBiasParameters.x;
+			}
+		}
+
+		// fade shadows in the distance
+		#if !RADIAL_ATTENUATION 
+			float Depth = dot(WorldPositionForLighting - View.ViewOrigin.xyz, View.ViewForward);
+			float DistanceFade = saturate(Depth * ShadowInjectParams.z + ShadowInjectParams.w);
+			DynamicShadowFactor = lerp(DynamicShadowFactor, 1.0f, DistanceFade * DistanceFade);
+		#endif
+
+		// Combine static shadowing and dynamic shadowing, important for stationary directional lights with CSM
+		ShadowFactor = min(ShadowFactor, DynamicShadowFactor);
+
 	#endif
 
 	return ShadowFactor;
@@ -130,7 +230,7 @@ void InjectMainPS(
 	OutColor0 = 0;
 	OutColor1 = 0;
 
-	// compute XYZ of the position we shader
+	// compute XYZ of the position we want to shade
 	float3 WorldPosition;
 	{
 		float ZPosition = View.TranslucencyLightingVolumeMin[VolumeCascadeIndex].z + (Input.LayerIndex + .5f) * View.TranslucencyLightingVolumeInvSize[VolumeCascadeIndex].w;
@@ -144,11 +244,12 @@ void InjectMainPS(
 	// 0: no contribution, 1:full contribution
 	float Masking = 1.0f;
 
-	// clip if not in cascade shaft (two planes)
+	
 	#if RADIAL_ATTENUATION 
 	{
-		// cull volumes outside the light radius (value is < 0)
-		clip(CascadeBounds.w * CascadeBounds.w - dot(CascadeBounds.xyz - WorldPosition, CascadeBounds.xyz - WorldPosition));
+		// cull voxels outside the light radius (value is < 0)
+		float3 LightVector = DeferredLightUniforms.LightPosition - WorldPosition;
+		clip(1.0f / (DeferredLightUniforms.LightInvRadius * DeferredLightUniforms.LightInvRadius) - dot(LightVector, LightVector));
 	}
 	#else
 	{
@@ -166,17 +267,6 @@ void InjectMainPS(
 	}
 	#endif
 
-	// todo: use DistanceFromCameraFade()
-	// Returns 0 for positions closer than the fade near distance from the camera, and 1 for positions further than the fade far distance.
-	float DistanceFade = 0;
-	#if !RADIAL_ATTENUATION
-	{
-		float Depth = dot(WorldPosition - View.ViewOrigin.xyz, View.ViewForward);
-		DistanceFade = saturate(Depth * ShadowInjectParams.z + ShadowInjectParams.w);
-		DistanceFade = DistanceFade * DistanceFade;
-	}
-	#endif
-
 	float3 NormalizedLightVector = GetNormalizedLightVector(WorldPosition);
 
 	float3 WorldPositionForLighting;
@@ -192,35 +282,8 @@ void InjectMainPS(
 		// Anything else would cause the extents of the shadowmap to not match up with the cone falloff on a spotlight
 		float Attenuation = CalcLightAttenuation(WorldPositionForLighting, WorldLightVector);
 
-		float ShadowFactor = 1;
-
-		#if DYNAMICALLY_SHADOWED
-			bool bUseCubemapShadowing = false;
-
-			#if RADIAL_ATTENUATION
-				bUseCubemapShadowing = SpotlightMask < 1;
-			#endif
-
-			if (bUseCubemapShadowing)
-			{
-				ShadowFactor = CubemapHardwarePCF(WorldPositionForLighting, DeferredLightUniforms.LightPosition, DeferredLightUniforms.LightInvRadius, 0.03f * 512 * InvShadowmapResolution);
-			}
-			else
-			{
-				// Transform the world position into shadowmap space
-				float4 HomogeneousShadowPosition = mul(float4(WorldPositionForLighting, 1), WorldToShadowMatrix);
-				float2 ShadowUVs = HomogeneousShadowPosition.xy / HomogeneousShadowPosition.w;
-
-				// Treat as unshadowed if the voxel is outside of the shadow map
-				if (all(ShadowUVs >= ShadowmapMinMax.xy && ShadowUVs <= ShadowmapMinMax.zw))
-				{
-					// Sample the shadowmap depth and determine if this voxel is shadowed
-					float ShadowDepth = Texture2DSampleLevel(ShadowDepthTexture, ShadowDepthTextureSampler, ShadowUVs, 0).x;
-					ShadowFactor = HomogeneousShadowPosition.z < ShadowDepth - DepthBiasParameters.x;
-				}
-			}
-		#endif
-
+		float ShadowFactor = ComputeVolumeShadowing(WorldPositionForLighting);
+		
 		if (VolumeCascadeIndex == 1)
 		{
 			// Larger values result in a shorter transition distance
@@ -242,11 +305,6 @@ void InjectMainPS(
 			#endif
 		}
 
-		// fade shadows in the distance
-		#if !RADIAL_ATTENUATION 
-			ShadowFactor = lerp(ShadowFactor, 1.0f, DistanceFade);
-		#endif
-
 		// Apply light function after edge fading, so that a light function which is modifying the light's brightness still works
 		ShadowFactor *= GetLightFunctionShadowFactor(WorldPositionForLighting);
 
diff --git a/Engine/Source/Editor/UnrealEd/Private/Lightmass/Lightmass.cpp b/Engine/Source/Editor/UnrealEd/Private/Lightmass/Lightmass.cpp
--- a/Engine/Source/Editor/UnrealEd/Private/Lightmass/Lightmass.cpp
+++ b/Engine/Source/Editor/UnrealEd/Private/Lightmass/Lightmass.cpp
@@ -1905,10 +1905,10 @@ void FLightmassExporter::WriteSceneSettings( Lightmass::FSceneFileHeader& Scene
 		verify(GConfig->GetFloat(TEXT("DevOptions.StaticShadows"), TEXT("MaxTransitionDistanceWorldSpace"), Scene.ShadowSettings.MaxTransitionDistanceWorldSpace, GLightmassIni));
 		verify(GConfig->GetInt(TEXT("DevOptions.StaticShadows"), TEXT("ApproximateHighResTexelsPerMaxTransitionDistance"), Scene.ShadowSettings.ApproximateHighResTexelsPerMaxTransitionDistance, GLightmassIni));
 		verify(GConfig->GetInt(TEXT("DevOptions.StaticShadows"), TEXT("MinDistanceFieldUpsampleFactor"), Scene.ShadowSettings.MinDistanceFieldUpsampleFactor, GLightmassIni));
-		Scene.ShadowSettings.DominantShadowTransitionSampleDistanceX = 1;
-		Scene.ShadowSettings.DominantShadowTransitionSampleDistanceY = 1;
-		Scene.ShadowSettings.DominantShadowSuperSampleFactor = 1;
-		Scene.ShadowSettings.DominantShadowMaxSamples = 1;
+		verify(GConfig->GetFloat(TEXT("DevOptions.StaticShadows"), TEXT("StaticShadowDepthMapTransitionSampleDistanceX"), Scene.ShadowSettings.StaticShadowDepthMapTransitionSampleDistanceX, GLightmassIni));
+		verify(GConfig->GetFloat(TEXT("DevOptions.StaticShadows"), TEXT("StaticShadowDepthMapTransitionSampleDistanceY"), Scene.ShadowSettings.StaticShadowDepthMapTransitionSampleDistanceY, GLightmassIni));
+		verify(GConfig->GetInt(TEXT("DevOptions.StaticShadows"), TEXT("StaticShadowDepthMapSuperSampleFactor"), Scene.ShadowSettings.StaticShadowDepthMapSuperSampleFactor, GLightmassIni));
+		verify(GConfig->GetInt(TEXT("DevOptions.StaticShadows"), TEXT("StaticShadowDepthMapMaxSamples"), Scene.ShadowSettings.StaticShadowDepthMapMaxSamples, GLightmassIni));
 		verify(GConfig->GetFloat(TEXT("DevOptions.StaticShadows"), TEXT("MinUnoccludedFraction"), Scene.ShadowSettings.MinUnoccludedFraction, GLightmassIni));
 	}
 	{
@@ -2367,6 +2367,24 @@ bool FLightmassProcessor::ExecuteAmortizedMaterialExport()
 	return Exporter->WriteToMaterialChannel(Statistics);
 }
 
+void FLightmassProcessor::IssueStaticShadowDepthMapTask(const ULightComponent* Light, int32 EstimatedCost)
+{
+	if (Light->HasStaticShadowing() && !Light->HasStaticLighting())
+	{
+		NSwarm::FTaskSpecification NewTaskSpecification(Light->LightGuid, TEXT("StaticShadowDepthMaps"), NSwarm::JOB_TASK_FLAG_USE_DEFAULTS );
+		NewTaskSpecification.Cost = EstimatedCost;
+		int32 ErrorCode = Swarm.AddTask( NewTaskSpecification );
+		if( ErrorCode >= 0 )
+		{
+			NumTotalSwarmTasks++;
+		}
+		else
+		{
+			UE_LOG(LogLightmassSolver, Log,  TEXT("Error, AddTask for StaticShadowDepthMaps failed with error code %d"), ErrorCode );
+		}
+	}
+}
+
 bool FLightmassProcessor::BeginRun()
 {
 	{
@@ -2651,6 +2669,26 @@ bool FLightmassProcessor::BeginRun()
 			}
 		}
 
+		{
+			for (int32 LightIndex = 0; LightIndex < Exporter->DirectionalLights.Num(); LightIndex++)
+			{
+				const ULightComponent* Light = Exporter->DirectionalLights[LightIndex];
+				IssueStaticShadowDepthMapTask(Light, INT_MAX);
+			}
+			
+			for (int32 LightIndex = 0; LightIndex < Exporter->SpotLights.Num(); LightIndex++)
+			{
+				const ULightComponent* Light = Exporter->SpotLights[LightIndex];
+				IssueStaticShadowDepthMapTask(Light, 10000);
+			}
+
+			for (int32 LightIndex = 0; LightIndex < Exporter->PointLights.Num(); LightIndex++)
+			{
+				const ULightComponent* Light = Exporter->PointLights[LightIndex];
+				IssueStaticShadowDepthMapTask(Light, 10000);
+			}
+		}
+
 		// Add BSP mapping tasks.
 		for( int32 MappingIdx=0; (ErrorCode >= 0) && MappingIdx < Exporter->BSPSurfaceMappings.Num() && !GEditor->GetMapBuildCancelled(); MappingIdx++ )
 		{
@@ -3513,6 +3551,32 @@ FStaticLightingTextureMapping* FLightmassProcessor::GetStaticLightingTextureMapp
 	return NULL;
 }
 
+void FLightmassProcessor::ImportStaticShadowDepthMap(ULightComponent* Light)
+{
+	const FString ChannelName = Lightmass::CreateChannelName(Light->LightGuid, Lightmass::LM_DOMINANTSHADOW_VERSION, Lightmass::LM_DOMINANTSHADOW_EXTENSION);
+	const int32 Channel = Swarm.OpenChannel( *ChannelName, LM_DOMINANTSHADOW_CHANNEL_FLAGS );
+	if (Channel >= 0)
+	{
+		FStaticShadowDepthMap& DepthMap = Light->StaticShadowDepthMap;
+		Lightmass::FStaticShadowDepthMapData ShadowMapData;
+		Swarm.ReadChannel(Channel, &ShadowMapData, sizeof(ShadowMapData));
+
+		DepthMap.WorldToLight = ShadowMapData.WorldToLight;
+		DepthMap.ShadowMapSizeX = ShadowMapData.ShadowMapSizeX;
+		DepthMap.ShadowMapSizeY = ShadowMapData.ShadowMapSizeY;
+
+		ReadArray(Channel, DepthMap.DepthSamples);
+		Swarm.CloseChannel(Channel);
+
+		BeginReleaseResource(&DepthMap);
+		BeginInitResource(&DepthMap);
+	}
+	else
+	{
+		UE_LOG(LogLightmassSolver, Log,  TEXT("Error, OpenChannel failed to open %s with error code %d"), *ChannelName, Channel );
+	}
+}
+
 /**
  * Import the mapping specified by a Guid.
  *	@param MappingGuid				Guid that identifies a mapping
@@ -3530,10 +3594,19 @@ void FLightmassProcessor::ImportMapping( const FGuid& MappingGuid, bool bProcess
 	}
 	else
 	{
-		FMappingImportHelper** pImportData = ImportedMappings.Find(MappingGuid);
-		if ((pImportData == NULL) || (*pImportData == NULL))
+		ULightComponent* Light = FindLight(MappingGuid);
+
+		if (Light)
 		{
-			UE_LOG(LogLightmassSolver, Warning, TEXT("Mapping not found for %s"), *(MappingGuid.ToString()));
+			ImportStaticShadowDepthMap(Light);
+		}
+		else
+		{
+			FMappingImportHelper** pImportData = ImportedMappings.Find(MappingGuid);
+			if ((pImportData == NULL) || (*pImportData == NULL))
+			{
+				UE_LOG(LogLightmassSolver, Warning, TEXT("Mapping not found for %s"), *(MappingGuid.ToString()));
+			}
 		}
 	}
 
@@ -3694,7 +3767,7 @@ void FLightmassProcessor::ImportDebugOutput()
  *	@return	ULightComponent*	The corresponding light component.
  *								NULL if not found.
  */
-ULightComponent* FLightmassProcessor::FindLight(FGuid& LightGuid)
+ULightComponent* FLightmassProcessor::FindLight(const FGuid& LightGuid)
 {
 	if (Exporter)
 	{
diff --git a/Engine/Source/Editor/UnrealEd/Private/Lightmass/Lightmass.h b/Engine/Source/Editor/UnrealEd/Private/Lightmass/Lightmass.h
--- a/Engine/Source/Editor/UnrealEd/Private/Lightmass/Lightmass.h
+++ b/Engine/Source/Editor/UnrealEd/Private/Lightmass/Lightmass.h
@@ -345,6 +345,8 @@ public:
 	/** Exports some materials, amortized, returning true if completely done */
 	bool ExecuteAmortizedMaterialExport();
 
+	void IssueStaticShadowDepthMapTask(const ULightComponent* Light, int32 EstimatedCost);
+
 	/** Starts the asynchronous lightmass process */
 	bool BeginRun();
 
@@ -494,7 +496,7 @@ protected:
 	 *	@return	ULightComponent*	The corresponding light component.
 	 *								NULL if not found.
 	 */
-	ULightComponent* FindLight(FGuid& LightGuid);
+	ULightComponent* FindLight(const FGuid& LightGuid);
 
 	/**
 	 *	Retrieve the static mehs for the given Guid
@@ -655,6 +657,8 @@ protected:
 	 */
 	void	ImportStaticLightingTextureMapping( const FGuid& MappingGuid, bool bProcessImmediately );
 
+	void ImportStaticShadowDepthMap(ULightComponent* Light);
+
 	/** Reads in a TArray from the given channel. */
 	template<class T>
 	void ReadArray(int32 Channel, TArray<T>& Array);
diff --git a/Engine/Source/Programs/UnrealLightmass/Private/ImportExport/Exporter.cpp b/Engine/Source/Programs/UnrealLightmass/Private/ImportExport/Exporter.cpp
--- a/Engine/Source/Programs/UnrealLightmass/Private/ImportExport/Exporter.cpp
+++ b/Engine/Source/Programs/UnrealLightmass/Private/ImportExport/Exporter.cpp
@@ -76,20 +76,20 @@ namespace Lightmass
 	}
 
 	/** Exports dominant shadow information to Unreal. */
-	void FLightmassSolverExporter::ExportDominantShadowInfo(const FGuid& LightGuid, const FDominantLightShadowInfo& DominantLightShadowInfo) const
+	void FLightmassSolverExporter::ExportStaticShadowDepthMap(const FGuid& LightGuid, const FStaticShadowDepthMap& StaticShadowDepthMap) const
 	{
 		const FString ChannelName = CreateChannelName(LightGuid, LM_DOMINANTSHADOW_VERSION, LM_DOMINANTSHADOW_EXTENSION);
 		const int32 ErrorCode = Swarm->OpenChannel(*ChannelName, LM_DOMINANTSHADOW_CHANNEL_FLAGS, true);
 		if( ErrorCode >= 0 )
 		{
-			Swarm->Write(&DominantLightShadowInfo, sizeof(FDominantLightShadowInfoData));
-			static_assert(sizeof(FDominantLightShadowSample) == sizeof(FDominantLightShadowSampleData), "Shadow derived size must match.");
-			WriteArray(DominantLightShadowInfo.ShadowMap);
+			Swarm->Write(&StaticShadowDepthMap, sizeof(FStaticShadowDepthMapData));
+			static_assert(sizeof(FStaticShadowDepthMapSample) == sizeof(FStaticShadowDepthMapSampleData), "ShadowDerivedSizeMustMatch");
+			WriteArray(StaticShadowDepthMap.ShadowMap);
 			Swarm->CloseCurrentChannel();
 		}
 		else
 		{
-			UE_LOG(LogLightmass, Log, TEXT("Failed to open dominant shadow channel!"));
+			UE_LOG(LogLightmass, Log, TEXT("Failed to open static shadow depth map channel!"));
 		}
 	}
 
diff --git a/Engine/Source/Programs/UnrealLightmass/Private/ImportExport/Exporter.h b/Engine/Source/Programs/UnrealLightmass/Private/ImportExport/Exporter.h
--- a/Engine/Source/Programs/UnrealLightmass/Private/ImportExport/Exporter.h
+++ b/Engine/Source/Programs/UnrealLightmass/Private/ImportExport/Exporter.h
@@ -48,7 +48,7 @@ namespace Lightmass
 			const TMap<FGuid,TArray<class FVolumeLightingSample> >& VolumeSamples) const;
 
 		/** Exports dominant shadow information to Unreal. */
-		void ExportDominantShadowInfo(const FGuid& LightGuid, const class FDominantLightShadowInfo& DominantLightShadowInfo) const;
+		void ExportStaticShadowDepthMap(const FGuid& LightGuid, const class FStaticShadowDepthMap& StaticShadowDepthMap) const;
 
 		/** 
 		 * Exports information about mesh area lights back to Unreal, 
diff --git a/Engine/Source/Programs/UnrealLightmass/Private/ImportExport/LightmassScene.cpp b/Engine/Source/Programs/UnrealLightmass/Private/ImportExport/LightmassScene.cpp
--- a/Engine/Source/Programs/UnrealLightmass/Private/ImportExport/LightmassScene.cpp
+++ b/Engine/Source/Programs/UnrealLightmass/Private/ImportExport/LightmassScene.cpp
@@ -416,8 +416,8 @@ void FScene::ApplyStaticLightingScale()
 	VolumeDistanceFieldSettings.VoxelSize *= SceneConstants.StaticLightingLevelScale;
 	VolumeDistanceFieldSettings.VolumeMaxDistance *= SceneConstants.StaticLightingLevelScale;
 	ShadowSettings.MaxTransitionDistanceWorldSpace *= SceneConstants.StaticLightingLevelScale;
-	ShadowSettings.DominantShadowTransitionSampleDistanceX *= SceneConstants.StaticLightingLevelScale;
-	ShadowSettings.DominantShadowTransitionSampleDistanceY *= SceneConstants.StaticLightingLevelScale;
+	ShadowSettings.StaticShadowDepthMapTransitionSampleDistanceX *= SceneConstants.StaticLightingLevelScale;
+	ShadowSettings.StaticShadowDepthMapTransitionSampleDistanceY *= SceneConstants.StaticLightingLevelScale;
 	PhotonMappingSettings.IndirectPhotonEmitDiskRadius *= SceneConstants.StaticLightingLevelScale;
 	PhotonMappingSettings.MaxImportancePhotonSearchDistance *= SceneConstants.StaticLightingLevelScale;
 	PhotonMappingSettings.MinImportancePhotonSearchDistance *= SceneConstants.StaticLightingLevelScale;
diff --git a/Engine/Source/Programs/UnrealLightmass/Private/Lighting/LightingSystem.cpp b/Engine/Source/Programs/UnrealLightmass/Private/Lighting/LightingSystem.cpp
--- a/Engine/Source/Programs/UnrealLightmass/Private/Lighting/LightingSystem.cpp
+++ b/Engine/Source/Programs/UnrealLightmass/Private/Lighting/LightingSystem.cpp
@@ -206,10 +206,6 @@ FStaticLightingSystem::FStaticLightingSystem(const FLightingBuildOptions& InOpti
 ,	MappingTasksInProgressThatWillNeedHelp(0)
 ,	bVolumeLightingSamplesComplete(false)
 ,	VolumeLightingInterpolationOctree(FVector4(0,0,0), HALF_WORLD_MAX)
-,	bDominantShadowTaskComplete(false)
-,	NumOutstandingDominantShadowColumns(0)
-,	OutstandingDominantShadowYIndex(-1)
-,	DominantDirectionalLightId(0,0,0,0)
 ,	bShouldExportMeshAreaLightData(false)
 ,	bShouldExportVolumeDistanceField(false)
 ,	NumPhotonsEmittedDirect(0)
@@ -672,39 +668,28 @@ void FStaticLightingSystem::ExportNonMappingTasks()
 	CompleteVisibilityTaskList.ApplyAndClear(*this);
 
 	{
-		TMap<const FSpotLight*,FDominantLightShadowInfo> DominantSpotLightShadowInfoCopy;
+		TMap<const FLight*, FStaticShadowDepthMap*> CompletedStaticShadowDepthMapsCopy;
 		{
 			// Enter a critical section before modifying DominantSpotLightShadowInfos since the worker threads may also modify it at any time
-			FScopeLock Lock(&DominantLightShadowSync);
-			DominantSpotLightShadowInfoCopy = DominantSpotLightShadowInfos;
-			DominantSpotLightShadowInfos.Empty();
+			FScopeLock Lock(&CompletedStaticShadowDepthMapsSync);
+			CompletedStaticShadowDepthMapsCopy = CompletedStaticShadowDepthMaps;
+			CompletedStaticShadowDepthMaps.Empty();
 		}
-		for (TMap<const FSpotLight*,FDominantLightShadowInfo>::TIterator It(DominantSpotLightShadowInfoCopy); It; ++It)
+
+		for (TMap<const FLight*,FStaticShadowDepthMap*>::TIterator It(CompletedStaticShadowDepthMapsCopy); It; ++It)
 		{
-			const FSpotLight* SpotLight = It.Key();
-			Exporter.ExportDominantShadowInfo(SpotLight->Guid, It.Value());
+			const FLight* Light = It.Key();
+			Exporter.ExportStaticShadowDepthMap(Light->Guid, *It.Value());
+
 			// Tell Swarm the task is complete (if we're not in debugging mode).
-			if ( !IsDebugMode() )
+			if (!IsDebugMode())
 			{
 				FLightmassSwarm* Swarm = GetExporter().GetSwarm();
-				Swarm->TaskCompleted( SpotLight->Guid );
+				Swarm->TaskCompleted(Light->Guid);
 			}
-		}
-	}
-	
-	if (bDominantShadowTaskComplete)
-	{
-		Exporter.ExportDominantShadowInfo(DominantDirectionalLightId, DominantLightShadowInfo);
 
-		DominantLightShadowInfo.ShadowMap.Empty();
-
-		// Tell Swarm the task is complete (if we're not in debugging mode).
-		if ( !IsDebugMode() )
-		{
-			FLightmassSwarm* Swarm = GetExporter().GetSwarm();
-			Swarm->TaskCompleted( DominantDirectionalLightId );
+			delete It.Value();
 		}
-		bDominantShadowTaskComplete = 0;
 	}
 
 	if (bShouldExportMeshAreaLightData)
@@ -865,8 +850,8 @@ void FStaticLightingSystem::ValidateSettings(FScene& InScene)
 
 	// Round up to nearest odd number
 	ShadowSettings.MinDistanceFieldUpsampleFactor = FMath::Clamp(ShadowSettings.MinDistanceFieldUpsampleFactor - ShadowSettings.MinDistanceFieldUpsampleFactor % 2 + 1, 1, 17);
-	ShadowSettings.DominantShadowTransitionSampleDistanceX = FMath::Max(ShadowSettings.DominantShadowTransitionSampleDistanceX, DELTA);
-	ShadowSettings.DominantShadowTransitionSampleDistanceY = FMath::Max(ShadowSettings.DominantShadowTransitionSampleDistanceY, DELTA);
+	ShadowSettings.StaticShadowDepthMapTransitionSampleDistanceX = FMath::Max(ShadowSettings.StaticShadowDepthMapTransitionSampleDistanceX, DELTA);
+	ShadowSettings.StaticShadowDepthMapTransitionSampleDistanceY = FMath::Max(ShadowSettings.StaticShadowDepthMapTransitionSampleDistanceY, DELTA);
 
 	InScene.IrradianceCachingSettings.InterpolationMaxAngle = FMath::Clamp(InScene.IrradianceCachingSettings.InterpolationMaxAngle, 0.0f, 90.0f);
 	InScene.IrradianceCachingSettings.PointBehindRecordMaxAngle = FMath::Clamp(InScene.IrradianceCachingSettings.PointBehindRecordMaxAngle, 0.0f, 90.0f);
@@ -1019,9 +1004,9 @@ void FStaticLightingSystem::DumpStats(float TotalStaticLightingTime) const
 	{
 		SolverStats += FString::Printf( TEXT("%4.1f%%%8.1fs    Block on IrradianceCache Interpolation tasks\n"), 100.0f * Stats.BlockOnIndirectLightingInterpolateTasksTime / TotalLightingBusyThreadTime, Stats.BlockOnIndirectLightingInterpolateTasksTime);
 	}
-	if (Stats.DominantShadowThreadTime > 0.1f)
+	if (Stats.StaticShadowDepthMapThreadTime > 0.1f)
 	{
-		SolverStats += FString::Printf( TEXT("%4.1f%%%8.1fs    Dominant shadow map\n"), 100.0f * Stats.DominantShadowThreadTime / TotalLightingBusyThreadTime, Stats.DominantShadowThreadTime);
+		SolverStats += FString::Printf( TEXT("%4.1f%%%8.1fs    Static shadow depth maps (max %.1fs)\n"), 100.0f * Stats.StaticShadowDepthMapThreadTime / TotalLightingBusyThreadTime, Stats.StaticShadowDepthMapThreadTime, Stats.MaxStaticShadowDepthMapThreadTime);
 	}
 	if (Stats.VolumeDistanceFieldThreadTime > 0.1f)
 	{
@@ -1039,7 +1024,7 @@ void FStaticLightingSystem::DumpStats(float TotalStaticLightingTime) const
 	{
 		SolverStats += FString::Printf( TEXT("%4.1f%%%8.1fs    Volume Samples\n"), 100.0f * Stats.VolumeSampleThreadTime / TotalLightingBusyThreadTime, Stats.VolumeSampleThreadTime);
 	}
-	const float UnaccountedLightingThreadTime = FMath::Max(TotalLightingBusyThreadTime - (SampleSetupTime + Stats.DirectLightingTime + Stats.BlockOnIndirectLightingCacheTasksTime + Stats.BlockOnIndirectLightingInterpolateTasksTime + Stats.SecondPassIrradianceCacheInterpolationTime + Stats.VolumeSampleThreadTime + Stats.DominantShadowThreadTime + Stats.VolumeDistanceFieldThreadTime + PrecomputedVisibilityThreadTime), 0.0f);
+	const float UnaccountedLightingThreadTime = FMath::Max(TotalLightingBusyThreadTime - (SampleSetupTime + Stats.DirectLightingTime + Stats.BlockOnIndirectLightingCacheTasksTime + Stats.BlockOnIndirectLightingInterpolateTasksTime + Stats.SecondPassIrradianceCacheInterpolationTime + Stats.VolumeSampleThreadTime + Stats.StaticShadowDepthMapThreadTime + Stats.VolumeDistanceFieldThreadTime + PrecomputedVisibilityThreadTime), 0.0f);
 	SolverStats += FString::Printf( TEXT("%4.1f%%%8.1fs    Unaccounted\n"), 100.0f * UnaccountedLightingThreadTime / TotalLightingBusyThreadTime, UnaccountedLightingThreadTime);
 	// Send the message in multiple parts since it cuts off in the middle otherwise
 	LogSolverMessage(SolverStats);
@@ -1401,7 +1386,7 @@ FStaticLightingMapping*	FStaticLightingSystem::ThreadGetNextMapping(
 	bool& bWaitTimedOut, 
 	bool& bDynamicObjectTask, 
 	int32& PrecomputedVisibilityTaskIndex,
-	bool& bDominantShadowTask,
+	bool& bStaticShadowDepthMapTask,
 	bool& bMeshAreaLightDataTask,
 	bool& bVolumeDataTask)
 {
@@ -1411,7 +1396,7 @@ FStaticLightingMapping*	FStaticLightingSystem::ThreadGetNextMapping(
 	bWaitTimedOut = true;
 	bDynamicObjectTask = false;
 	PrecomputedVisibilityTaskIndex = INDEX_NONE;
-	bDominantShadowTask = false;
+	bStaticShadowDepthMapTask = false;
 	bMeshAreaLightDataTask = false;
 	bVolumeDataTask = false;
 
@@ -1465,7 +1450,7 @@ FStaticLightingMapping*	FStaticLightingSystem::ThreadGetNextMapping(
 			}
 			else if (Scene.FindLightByGuid(TaskGuid))
 			{
-				bDominantShadowTask = true;
+				bStaticShadowDepthMapTask = true;
 				Swarm->AcceptTask( TaskGuid );
 				bWaitTimedOut = false;
 			}
@@ -1506,22 +1491,7 @@ void FStaticLightingSystem::ThreadLoop(bool bIsMainThread, int32 ThreadIndex, FT
 	while (!bIsDone)
 	{
 		const double StartLoopTime = FPlatformTime::Seconds();
-		// Process another row of the dominant shadow map if necessary
-		if (NumOutstandingDominantShadowColumns > 0)
-		{
-			const int32 ThreadY = FPlatformAtomics::InterlockedIncrement(&OutstandingDominantShadowYIndex);
-			if (ThreadY < DominantLightShadowInfo.ShadowMapSizeY)
-			{
-				CalculateDominantShadowInfoWorkRange(ThreadY);
-				const int32 NumTasksRemaining = FPlatformAtomics::InterlockedDecrement(&NumOutstandingDominantShadowColumns);
-				if (NumTasksRemaining == 0)
-				{
-					// Signal to the main thread that all dominant shadow map tasks are complete
-					FPlatformAtomics::InterlockedExchange(&bDominantShadowTaskComplete, true);
-				}
-			}
-		}
-
+		
 		if (NumOutstandingVolumeDataLayers > 0)
 		{
 			const int32 ThreadZ = FPlatformAtomics::InterlockedIncrement(&OutstandingVolumeDataLayerIndex);
@@ -1541,7 +1511,7 @@ void FStaticLightingSystem::ThreadLoop(bool bIsMainThread, int32 ThreadIndex, FT
 		bool bRequestForTaskTimedOut;
 		bool bDynamicObjectTask;
 		int32 PrecomputedVisibilityTaskIndex;
-		bool bDominantShadowTask;
+		bool bStaticShadowDepthMapTask;
 		bool bMeshAreaLightDataTask;
 		bool bVolumeDataTask;
 
@@ -1553,7 +1523,7 @@ void FStaticLightingSystem::ThreadLoop(bool bIsMainThread, int32 ThreadIndex, FT
 			bRequestForTaskTimedOut, 
 			bDynamicObjectTask, 
 			PrecomputedVisibilityTaskIndex,
-			bDominantShadowTask,
+			bStaticShadowDepthMapTask,
 			bMeshAreaLightDataTask,
 			bVolumeDataTask);
 
@@ -1589,13 +1559,13 @@ void FStaticLightingSystem::ThreadLoop(bool bIsMainThread, int32 ThreadIndex, FT
 		{
 			BeginCalculateVolumeDistanceField();
 		}
-		else if (bDominantShadowTask)
+		else if (bStaticShadowDepthMapTask)
 		{
-			BeginCalculateDominantShadowInfo(TaskGuid);
+			CalculateStaticShadowDepthMap(TaskGuid);
 		}
 		else
 		{
-			if (!bSignaledMappingsComplete && NumOutstandingDominantShadowColumns <= 0 && NumOutstandingVolumeDataLayers <= 0)
+			if (!bSignaledMappingsComplete && NumOutstandingVolumeDataLayers <= 0)
 			{
 				bSignaledMappingsComplete = true;
 				GSwarm->SendMessage( NSwarm::FTimingMessage( NSwarm::PROGSTATE_Processing0, ThreadIndex ) );
@@ -1622,7 +1592,6 @@ void FStaticLightingSystem::ThreadLoop(bool bIsMainThread, int32 ThreadIndex, FT
 
 			if (!NextCacheTask 
 				&& !NextInterpolateTask
-				&& NumOutstandingDominantShadowColumns <= 0 
 				&& NumOutstandingVolumeDataLayers <= 0)
 			{
 				if (MappingTasksInProgressThatWillNeedHelp <= 0 && !bRequestForTaskTimedOut)
@@ -2074,95 +2043,153 @@ FGatheredLightSample FStaticLightingSystem::CalculateApproximateSkyLighting(
 	return IncomingRadiance;
 }
 
-/** Initializes DominantLightShadowInfo and prepares for multithreaded generation of DominantLightShadowInfo.ShadowMap. */
-void FStaticLightingSystem::BeginCalculateDominantShadowInfo(FGuid LightGuid)
+void FStaticLightingSystem::CalculateStaticShadowDepthMap(FGuid LightGuid)
 {
-	const FLight* DominantLight = Scene.FindLightByGuid(LightGuid);
-	check(DominantLight);
-	const FDirectionalLight* DominantDirectionalLight = DominantLight->GetDirectionalLight();
-	const FSpotLight* DominantSpotLight = DominantLight->GetSpotLight();
-	
-	if (DominantDirectionalLight)
+	const FLight* Light = Scene.FindLightByGuid(LightGuid);
+	check(Light);
+	const FDirectionalLight* DirectionalLight = Light->GetDirectionalLight();
+	const FSpotLight* SpotLight = Light->GetSpotLight();
+	const FPointLight* PointLight = Light->GetPointLight();
+	check(DirectionalLight || SpotLight || PointLight);
+
+	const double StartTime = FPlatformTime::Seconds();
+
+	FStaticLightingMappingContext Context(NULL, *this);
+	FStaticShadowDepthMap* ShadowDepthMap = new FStaticShadowDepthMap();
+
+	if (DirectionalLight)
 	{
-		checkSlow(DominantLightShadowInfo.ShadowMap.Num() == 0);
 		FVector4 XAxis, YAxis;
-		DominantDirectionalLight->Direction.FindBestAxisVectors3(XAxis, YAxis);
+		DirectionalLight->Direction.FindBestAxisVectors3(XAxis, YAxis);
 		// Create a coordinate system for the dominant directional light, with the z axis corresponding to the light's direction
-		DominantLightShadowInfo.WorldToLight = FBasisVectorMatrix(XAxis, YAxis, DominantDirectionalLight->Direction, FVector4(0,0,0));
-		FBoxSphereBounds ImportanceVolume = GetImportanceBounds().SphereRadius > 0.0f ? GetImportanceBounds() : FBoxSphereBounds(AggregateMesh.GetBounds());
-		const FBox LightSpaceImportanceBounds = ImportanceVolume.GetBox().TransformBy(DominantLightShadowInfo.WorldToLight);
+		ShadowDepthMap->WorldToLight = FBasisVectorMatrix(XAxis, YAxis, DirectionalLight->Direction, FVector4(0,0,0));
 
-		DominantLightShadowInfo.LightSpaceImportanceBoundMin = LightSpaceImportanceBounds.Min;
-		DominantLightShadowInfo.LightSpaceImportanceBoundMax = LightSpaceImportanceBounds.Max;
+		FBoxSphereBounds ImportanceVolume = GetImportanceBounds().SphereRadius > 0.0f ? GetImportanceBounds() : FBoxSphereBounds(AggregateMesh.GetBounds());
+		const FBox LightSpaceImportanceBounds = ImportanceVolume.GetBox().TransformBy(ShadowDepthMap->WorldToLight);
 
-		DominantLightShadowInfo.ShadowMapSizeX = FMath::TruncToInt(FMath::Max(LightSpaceImportanceBounds.GetExtent().X * 2.0f / ShadowSettings.DominantShadowTransitionSampleDistanceX, 4.0f));
-		DominantLightShadowInfo.ShadowMapSizeX = DominantLightShadowInfo.ShadowMapSizeX == appTruncErrorCode ? INT_MAX : DominantLightShadowInfo.ShadowMapSizeX;
-		DominantLightShadowInfo.ShadowMapSizeY = FMath::TruncToInt(FMath::Max(LightSpaceImportanceBounds.GetExtent().Y * 2.0f / ShadowSettings.DominantShadowTransitionSampleDistanceY, 4.0f));
-		DominantLightShadowInfo.ShadowMapSizeY = DominantLightShadowInfo.ShadowMapSizeY == appTruncErrorCode ? INT_MAX : DominantLightShadowInfo.ShadowMapSizeY;
+		ShadowDepthMap->ShadowMapSizeX = FMath::TruncToInt(FMath::Max(LightSpaceImportanceBounds.GetExtent().X * 2.0f / ShadowSettings.StaticShadowDepthMapTransitionSampleDistanceX, 4.0f));
+		ShadowDepthMap->ShadowMapSizeX = ShadowDepthMap->ShadowMapSizeX == appTruncErrorCode ? INT_MAX : ShadowDepthMap->ShadowMapSizeX;
+		ShadowDepthMap->ShadowMapSizeY = FMath::TruncToInt(FMath::Max(LightSpaceImportanceBounds.GetExtent().Y * 2.0f / ShadowSettings.StaticShadowDepthMapTransitionSampleDistanceY, 4.0f));
+		ShadowDepthMap->ShadowMapSizeY = ShadowDepthMap->ShadowMapSizeY == appTruncErrorCode ? INT_MAX : ShadowDepthMap->ShadowMapSizeY;
 
 		// Clamp the number of dominant shadow samples generated if necessary while maintaining aspect ratio
-		if ((uint64)DominantLightShadowInfo.ShadowMapSizeX * (uint64)DominantLightShadowInfo.ShadowMapSizeY > (uint64)ShadowSettings.DominantShadowMaxSamples)
+		if ((uint64)ShadowDepthMap->ShadowMapSizeX * (uint64)ShadowDepthMap->ShadowMapSizeY > (uint64)ShadowSettings.StaticShadowDepthMapMaxSamples)
 		{
-			const float AspectRatio = DominantLightShadowInfo.ShadowMapSizeX / (float)DominantLightShadowInfo.ShadowMapSizeY;
-			DominantLightShadowInfo.ShadowMapSizeY = FMath::TruncToInt(FMath::Sqrt(ShadowSettings.DominantShadowMaxSamples / AspectRatio));
-			DominantLightShadowInfo.ShadowMapSizeX = FMath::TruncToInt(ShadowSettings.DominantShadowMaxSamples / DominantLightShadowInfo.ShadowMapSizeY);
+			const float AspectRatio = ShadowDepthMap->ShadowMapSizeX / (float)ShadowDepthMap->ShadowMapSizeY;
+			ShadowDepthMap->ShadowMapSizeY = FMath::TruncToInt(FMath::Sqrt(ShadowSettings.StaticShadowDepthMapMaxSamples / AspectRatio));
+			ShadowDepthMap->ShadowMapSizeX = FMath::TruncToInt(ShadowSettings.StaticShadowDepthMapMaxSamples / ShadowDepthMap->ShadowMapSizeY);
 		}
 
 		// Allocate the shadow map
-		DominantLightShadowInfo.ShadowMap.Empty(DominantLightShadowInfo.ShadowMapSizeX * DominantLightShadowInfo.ShadowMapSizeY);
-		DominantLightShadowInfo.ShadowMap.AddZeroed(DominantLightShadowInfo.ShadowMapSizeX * DominantLightShadowInfo.ShadowMapSizeY);
+		ShadowDepthMap->ShadowMap.Empty(ShadowDepthMap->ShadowMapSizeX * ShadowDepthMap->ShadowMapSizeY);
+		ShadowDepthMap->ShadowMap.AddZeroed(ShadowDepthMap->ShadowMapSizeX * ShadowDepthMap->ShadowMapSizeY);
 
-		DominantDirectionalLightId = DominantDirectionalLight->Guid;
+		{
+			const float InvDistanceRange = 1.0f / (LightSpaceImportanceBounds.Max.Z - LightSpaceImportanceBounds.Min.Z);
+			const FMatrix LightToWorld = ShadowDepthMap->WorldToLight.Inverse();
 
-		// Signal to the other threads to start processing rows of the shadow map
-		FPlatformAtomics::InterlockedExchange(&NumOutstandingDominantShadowColumns, DominantLightShadowInfo.ShadowMapSizeY);
+			for (int32 Y = 0; Y < ShadowDepthMap->ShadowMapSizeY; Y++)
+			{
+				for (int32 X = 0; X < ShadowDepthMap->ShadowMapSizeX; X++)
+				{
+					float MaxSampleDistance = 0.0f;
+					// Super sample each cell
+					for (int32 SubSampleY = 0; SubSampleY < ShadowSettings.StaticShadowDepthMapSuperSampleFactor; SubSampleY++)
+					{
+						const float YFraction = (Y + SubSampleY / (float)ShadowSettings.StaticShadowDepthMapSuperSampleFactor) / (float)(ShadowDepthMap->ShadowMapSizeY - 1);
+						for (int32 SubSampleX = 0; SubSampleX < ShadowSettings.StaticShadowDepthMapSuperSampleFactor; SubSampleX++)
+						{
+							const float XFraction = (X + SubSampleX / (float)ShadowSettings.StaticShadowDepthMapSuperSampleFactor) / (float)(ShadowDepthMap->ShadowMapSizeX - 1);
+							// Construct a ray in light space along the direction of the light, starting at the minimum light space Z going to the maximum.
+							const FVector4 LightSpaceStartPosition(
+								LightSpaceImportanceBounds.Min.X + XFraction * (LightSpaceImportanceBounds.Max.X - LightSpaceImportanceBounds.Min.X),
+								LightSpaceImportanceBounds.Min.Y + YFraction * (LightSpaceImportanceBounds.Max.Y - LightSpaceImportanceBounds.Min.Y),
+								LightSpaceImportanceBounds.Min.Z);
+							const FVector4 LightSpaceEndPosition(LightSpaceStartPosition.X, LightSpaceStartPosition.Y, LightSpaceImportanceBounds.Max.Z);
+							// Transform the ray into world space in order to trace against the world space aggregate mesh
+							const FVector4 WorldSpaceStartPosition = LightToWorld.TransformPosition(LightSpaceStartPosition);
+							const FVector4 WorldSpaceEndPosition = LightToWorld.TransformPosition(LightSpaceEndPosition);
+							const FLightRay LightRay(
+								WorldSpaceStartPosition,
+								WorldSpaceEndPosition,
+								NULL,
+								NULL,
+								// We are tracing from the light instead of to the light,
+								// So flip sidedness so that backface culling matches up with tracing to the light
+								LIGHTRAY_FLIP_SIDEDNESS
+								);
+
+							FLightRayIntersection Intersection;
+							AggregateMesh.IntersectLightRay(LightRay, true, false, true, Context.RayCache, Intersection);
+
+							if (Intersection.bIntersects)
+							{
+								// Use the maximum distance of all super samples for each cell, to get a conservative shadow map
+								MaxSampleDistance = FMath::Max(MaxSampleDistance, (Intersection.IntersectionVertex.WorldPosition - WorldSpaceStartPosition).Size3());
+							}
+						}
+					}
+
+					if (MaxSampleDistance == 0.0f)
+					{
+						MaxSampleDistance = LightSpaceImportanceBounds.Max.Z - LightSpaceImportanceBounds.Min.Z;
+					} 
+
+					ShadowDepthMap->ShadowMap[Y * ShadowDepthMap->ShadowMapSizeX + X] = FStaticShadowDepthMapSample(FFloat16(MaxSampleDistance * InvDistanceRange));
+				}
+			}
+		}
+
+		ShadowDepthMap->WorldToLight *= FTranslationMatrix(-LightSpaceImportanceBounds.Min)
+			* FScaleMatrix(FVector(1.0f) / (LightSpaceImportanceBounds.Max - LightSpaceImportanceBounds.Min));
+
+		FScopeLock Lock(&CompletedStaticShadowDepthMapsSync);
+		CompletedStaticShadowDepthMaps.Add(DirectionalLight, ShadowDepthMap);
 	}
-	else if (DominantSpotLight)
+	else if (SpotLight)
 	{
-		// Don't bother parallelizing spot light computations, spotlight shadowmaps are usually much smaller than directional lights since they don't affect the whole scene
-		FDominantLightShadowInfo DominantLightShadowInfo;
 		FVector4 XAxis, YAxis;
-		DominantSpotLight->Direction.FindBestAxisVectors3(XAxis, YAxis);
-		// Create a coordinate system for the dominant spot light, with the z axis corresponding to the light's direction, and translated to the light's origin
-		DominantLightShadowInfo.WorldToLight = FTranslationMatrix(-DominantSpotLight->Position) * FBasisVectorMatrix(XAxis, YAxis, DominantSpotLight->Direction, FVector4(0,0,0));
+		SpotLight->Direction.FindBestAxisVectors3(XAxis, YAxis);
+		// Create a coordinate system for the spot light, with the z axis corresponding to the light's direction, and translated to the light's origin
+		ShadowDepthMap->WorldToLight = FTranslationMatrix(-SpotLight->Position) 
+			* FBasisVectorMatrix(XAxis, YAxis, SpotLight->Direction, FVector4(0,0,0));
 
 		// Distance from the light's direction axis to the edge of the cone at the radius of the light
-		const float HalfCrossSectionLength = DominantSpotLight->Radius * FMath::Tan(DominantSpotLight->OuterConeAngle * (float)PI / 180.0f);
+		const float HalfCrossSectionLength = SpotLight->Radius * FMath::Tan(SpotLight->OuterConeAngle * (float)PI / 180.0f);
 
-		DominantLightShadowInfo.LightSpaceImportanceBoundMin = FVector4(-HalfCrossSectionLength, -HalfCrossSectionLength, 0);
-		DominantLightShadowInfo.LightSpaceImportanceBoundMax = FVector4(HalfCrossSectionLength, HalfCrossSectionLength, DominantSpotLight->Radius);
+		const FVector4 LightSpaceImportanceBoundMin = FVector4(-HalfCrossSectionLength, -HalfCrossSectionLength, 0);
+		const FVector4 LightSpaceImportanceBoundMax = FVector4(HalfCrossSectionLength, HalfCrossSectionLength, SpotLight->Radius);
 
-		DominantLightShadowInfo.ShadowMapSizeX = FMath::TruncToInt(FMath::Max(HalfCrossSectionLength / ShadowSettings.DominantShadowTransitionSampleDistanceX, 4.0f));
-		DominantLightShadowInfo.ShadowMapSizeX = DominantLightShadowInfo.ShadowMapSizeX == appTruncErrorCode ? INT_MAX : DominantLightShadowInfo.ShadowMapSizeX;
-		DominantLightShadowInfo.ShadowMapSizeY = DominantLightShadowInfo.ShadowMapSizeX;
+		ShadowDepthMap->ShadowMapSizeX = FMath::TruncToInt(FMath::Max(HalfCrossSectionLength / (2 * ShadowSettings.StaticShadowDepthMapTransitionSampleDistanceX), 4.0f));
+		ShadowDepthMap->ShadowMapSizeX = ShadowDepthMap->ShadowMapSizeX == appTruncErrorCode ? INT_MAX : ShadowDepthMap->ShadowMapSizeX;
+		ShadowDepthMap->ShadowMapSizeY = ShadowDepthMap->ShadowMapSizeX;
 
-		DominantLightShadowInfo.ShadowMap.Empty(DominantLightShadowInfo.ShadowMapSizeX * DominantLightShadowInfo.ShadowMapSizeY);
-		DominantLightShadowInfo.ShadowMap.AddZeroed(DominantLightShadowInfo.ShadowMapSizeX * DominantLightShadowInfo.ShadowMapSizeY);
+		ShadowDepthMap->ShadowMap.Empty(ShadowDepthMap->ShadowMapSizeX * ShadowDepthMap->ShadowMapSizeY);
+		ShadowDepthMap->ShadowMap.AddZeroed(ShadowDepthMap->ShadowMapSizeX * ShadowDepthMap->ShadowMapSizeY);
 
 		// Calculate the maximum possible distance for quantization
-		const float MaxPossibleDistance = FMath::Max(FVector4(HalfCrossSectionLength, HalfCrossSectionLength, DominantLightShadowInfo.LightSpaceImportanceBoundMax.Z - DominantLightShadowInfo.LightSpaceImportanceBoundMin.Z).Size3(), (float)KINDA_SMALL_NUMBER);
-		const FMatrix LightToWorld = DominantLightShadowInfo.WorldToLight.Inverse();
+		const float MaxPossibleDistance = LightSpaceImportanceBoundMax.Z - LightSpaceImportanceBoundMin.Z;
+		const FMatrix LightToWorld = ShadowDepthMap->WorldToLight.Inverse();
 		const FBoxSphereBounds ImportanceVolume = GetImportanceBounds().SphereRadius > 0.0f ? GetImportanceBounds() : FBoxSphereBounds(AggregateMesh.GetBounds());
 
-		FStaticLightingMappingContext Context(NULL, *this);
-		for (int32 Y = 0; Y < DominantLightShadowInfo.ShadowMapSizeY; Y++)
+		for (int32 Y = 0; Y < ShadowDepthMap->ShadowMapSizeY; Y++)
 		{
-			for (int32 X = 0; X < DominantLightShadowInfo.ShadowMapSizeX; X++)
+			for (int32 X = 0; X < ShadowDepthMap->ShadowMapSizeX; X++)
 			{
 				float MaxSampleDistance = 0.0f;
 				// Super sample each cell
-				for (int32 SubSampleY = 0; SubSampleY < ShadowSettings.DominantShadowSuperSampleFactor; SubSampleY++)
+				for (int32 SubSampleY = 0; SubSampleY < ShadowSettings.StaticShadowDepthMapSuperSampleFactor; SubSampleY++)
 				{
-					const float YFraction = (Y + SubSampleY / (float)ShadowSettings.DominantShadowSuperSampleFactor) / (float)(DominantLightShadowInfo.ShadowMapSizeY - 1);
-					for (int32 SubSampleX = 0; SubSampleX < ShadowSettings.DominantShadowSuperSampleFactor; SubSampleX++)
+					const float YFraction = (Y + SubSampleY / (float)ShadowSettings.StaticShadowDepthMapSuperSampleFactor) / (float)(ShadowDepthMap->ShadowMapSizeY - 1);
+					for (int32 SubSampleX = 0; SubSampleX < ShadowSettings.StaticShadowDepthMapSuperSampleFactor; SubSampleX++)
 					{
-						const float XFraction = (X + SubSampleX / (float)ShadowSettings.DominantShadowSuperSampleFactor) / (float)(DominantLightShadowInfo.ShadowMapSizeX - 1);
+						const float XFraction = (X + SubSampleX / (float)ShadowSettings.StaticShadowDepthMapSuperSampleFactor) / (float)(ShadowDepthMap->ShadowMapSizeX - 1);
 						// Construct a ray in light space along the direction of the light, starting at the light and going to the maximum light space Z.
 						const FVector4 LightSpaceStartPosition(0,0,0);
 						const FVector4 LightSpaceEndPosition(
-							DominantLightShadowInfo.LightSpaceImportanceBoundMin.X + XFraction * (DominantLightShadowInfo.LightSpaceImportanceBoundMax.X - DominantLightShadowInfo.LightSpaceImportanceBoundMin.X),
-							DominantLightShadowInfo.LightSpaceImportanceBoundMin.Y + YFraction * (DominantLightShadowInfo.LightSpaceImportanceBoundMax.Y - DominantLightShadowInfo.LightSpaceImportanceBoundMin.Y),
-							DominantLightShadowInfo.LightSpaceImportanceBoundMax.Z);
+							LightSpaceImportanceBoundMin.X + XFraction * (LightSpaceImportanceBoundMax.X - LightSpaceImportanceBoundMin.X),
+							LightSpaceImportanceBoundMin.Y + YFraction * (LightSpaceImportanceBoundMax.Y - LightSpaceImportanceBoundMin.Y),
+							LightSpaceImportanceBoundMax.Z);
 						// Transform the ray into world space in order to trace against the world space aggregate mesh
 						const FVector4 WorldSpaceStartPosition = LightToWorld.TransformPosition(LightSpaceStartPosition);
 						const FVector4 WorldSpaceEndPosition = LightToWorld.TransformPosition(LightSpaceEndPosition);
@@ -2178,10 +2205,12 @@ void FStaticLightingSystem::BeginCalculateDominantShadowInfo(FGuid LightGuid)
 
 						FLightRayIntersection Intersection;
 						AggregateMesh.IntersectLightRay(LightRay, true, false, true, Context.RayCache, Intersection);
-						if (Intersection.bIntersects && ImportanceVolume.GetBox().IsInside(Intersection.IntersectionVertex.WorldPosition))
+
+						if (Intersection.bIntersects)
 						{
+							const FVector4 LightSpaceIntersectPosition = ShadowDepthMap->WorldToLight.TransformPosition(Intersection.IntersectionVertex.WorldPosition);
 							// Use the maximum distance of all super samples for each cell, to get a conservative shadow map
-							MaxSampleDistance = FMath::Max(MaxSampleDistance, (Intersection.IntersectionVertex.WorldPosition - WorldSpaceStartPosition).Size3());
+							MaxSampleDistance = FMath::Max(MaxSampleDistance, LightSpaceIntersectPosition.Z);
 						}
 					}
 				}
@@ -2191,76 +2220,100 @@ void FStaticLightingSystem::BeginCalculateDominantShadowInfo(FGuid LightGuid)
 					MaxSampleDistance = MaxPossibleDistance;
 				}
 
-				// Quantize the distance into a uint16 to reduce memory usage in Unreal
-				const uint16 QuantizedDistance = FMath::TruncToInt(65535 * FMath::Clamp(MaxSampleDistance / MaxPossibleDistance, 0.0f, 1.0f));
-				DominantLightShadowInfo.ShadowMap[Y * DominantLightShadowInfo.ShadowMapSizeX + X] = FDominantLightShadowSample(QuantizedDistance);
+				ShadowDepthMap->ShadowMap[Y * ShadowDepthMap->ShadowMapSizeX + X] = FStaticShadowDepthMapSample(FFloat16(MaxSampleDistance / MaxPossibleDistance));
 			}
 		}
 
-		// Enter a critical section before modifying DominantSpotLightShadowInfos since the main thread may also modify it at any time
-		FScopeLock Lock(&DominantLightShadowSync);
-		DominantSpotLightShadowInfos.Add(DominantSpotLight, DominantLightShadowInfo);
-	}
-}
+		ShadowDepthMap->WorldToLight *= 
+			// Perspective projection sized to the spotlight cone
+			FPerspectiveMatrix(SpotLight->OuterConeAngle * (float)PI / 180.0f, 1, 1, 0, SpotLight->Radius)
+			// Convert from NDC to texture space, normalize Z
+			* FMatrix(
+				FPlane(.5f,								0,							0,									0),
+				FPlane(0,								.5f,						0,									0),
+				FPlane(0,								0,							1.0f / LightSpaceImportanceBoundMax.Z,		0),
+				FPlane(.5f,								.5f,						0,									1));
 
-/** Generates a single row of the dominant light shadow map. */
-void FStaticLightingSystem::CalculateDominantShadowInfoWorkRange(int32 ShadowMapY)
-{
-	const double CalculateWorkRangeStart = FPlatformTime::Seconds();
-	const float InvDistanceRange = 1.0f / (DominantLightShadowInfo.LightSpaceImportanceBoundMax.Z - DominantLightShadowInfo.LightSpaceImportanceBoundMin.Z);
-	const FMatrix LightToWorld = DominantLightShadowInfo.WorldToLight.Inverse();
-	FStaticLightingMappingContext Context(NULL, *this);
-	const FBoxSphereBounds ImportanceVolume = GetImportanceBounds().SphereRadius > 0.0f ? GetImportanceBounds() : FBoxSphereBounds(AggregateMesh.GetBounds());
-	// Process a single row of the shadow map
-	for (int32 X = 0; X < DominantLightShadowInfo.ShadowMapSizeX; X++)
+		FScopeLock Lock(&CompletedStaticShadowDepthMapsSync);
+		CompletedStaticShadowDepthMaps.Add(SpotLight, ShadowDepthMap);
+	}
+	else if (PointLight)
 	{
-		float MaxSampleDistance = 0.0f;
-		// Super sample each cell
-		for (int32 SubSampleY = 0; SubSampleY < ShadowSettings.DominantShadowSuperSampleFactor; SubSampleY++)
+		ShadowDepthMap->ShadowMapSizeX = FMath::TruncToInt(FMath::Max(PointLight->Radius / (2 * ShadowSettings.StaticShadowDepthMapTransitionSampleDistanceX), 4.0f));
+		ShadowDepthMap->ShadowMapSizeX = ShadowDepthMap->ShadowMapSizeX == appTruncErrorCode ? INT_MAX : ShadowDepthMap->ShadowMapSizeX;
+		ShadowDepthMap->ShadowMapSizeY = ShadowDepthMap->ShadowMapSizeX;
+
+		// Clamp the number of dominant shadow samples generated if necessary while maintaining aspect ratio
+		if ((uint64)ShadowDepthMap->ShadowMapSizeX * (uint64)ShadowDepthMap->ShadowMapSizeY > (uint64)ShadowSettings.StaticShadowDepthMapMaxSamples)
+		{
+			const float AspectRatio = ShadowDepthMap->ShadowMapSizeX / (float)ShadowDepthMap->ShadowMapSizeY;
+			ShadowDepthMap->ShadowMapSizeY = FMath::TruncToInt(FMath::Sqrt(ShadowSettings.StaticShadowDepthMapMaxSamples / AspectRatio));
+			ShadowDepthMap->ShadowMapSizeX = FMath::TruncToInt(ShadowSettings.StaticShadowDepthMapMaxSamples / ShadowDepthMap->ShadowMapSizeY);
+		}
+
+		// Allocate the shadow map
+		ShadowDepthMap->ShadowMap.Empty(ShadowDepthMap->ShadowMapSizeX * ShadowDepthMap->ShadowMapSizeY);
+		ShadowDepthMap->ShadowMap.AddZeroed(ShadowDepthMap->ShadowMapSizeX * ShadowDepthMap->ShadowMapSizeY);
+
+		ShadowDepthMap->WorldToLight = FMatrix::Identity;
+
+		for (int32 Y = 0; Y < ShadowDepthMap->ShadowMapSizeY; Y++)
 		{
-			const float YFraction = (ShadowMapY + SubSampleY / (float)ShadowSettings.DominantShadowSuperSampleFactor) / (float)(DominantLightShadowInfo.ShadowMapSizeY - 1);
-			for (int32 SubSampleX = 0; SubSampleX < ShadowSettings.DominantShadowSuperSampleFactor; SubSampleX++)
+			for (int32 X = 0; X < ShadowDepthMap->ShadowMapSizeX; X++)
 			{
-				const float XFraction = (X + SubSampleX / (float)ShadowSettings.DominantShadowSuperSampleFactor) / (float)(DominantLightShadowInfo.ShadowMapSizeX - 1);
-				// Construct a ray in light space along the direction of the light, starting at the minimum light space Z going to the maximum.
-				const FVector4 LightSpaceStartPosition(
-					DominantLightShadowInfo.LightSpaceImportanceBoundMin.X + XFraction * (DominantLightShadowInfo.LightSpaceImportanceBoundMax.X - DominantLightShadowInfo.LightSpaceImportanceBoundMin.X),
-					DominantLightShadowInfo.LightSpaceImportanceBoundMin.Y + YFraction * (DominantLightShadowInfo.LightSpaceImportanceBoundMax.Y - DominantLightShadowInfo.LightSpaceImportanceBoundMin.Y),
-					DominantLightShadowInfo.LightSpaceImportanceBoundMin.Z);
-				const FVector4 LightSpaceEndPosition(LightSpaceStartPosition.X, LightSpaceStartPosition.Y,DominantLightShadowInfo.LightSpaceImportanceBoundMax.Z);
-				// Transform the ray into world space in order to trace against the world space aggregate mesh
-				const FVector4 WorldSpaceStartPosition = LightToWorld.TransformPosition(LightSpaceStartPosition);
-				const FVector4 WorldSpaceEndPosition = LightToWorld.TransformPosition(LightSpaceEndPosition);
-				const FLightRay LightRay(
-					WorldSpaceStartPosition,
-					WorldSpaceEndPosition,
-					NULL,
-					NULL,
-					// We are tracing from the light instead of to the light,
-					// So flip sidedness so that backface culling matches up with tracing to the light
-					LIGHTRAY_FLIP_SIDEDNESS
-					);
-
-				FLightRayIntersection Intersection;
-				AggregateMesh.IntersectLightRay(LightRay, true, false, true, Context.RayCache, Intersection);
-				if (Intersection.bIntersects)
+				float MaxSampleDistance = 0.0f;
+				// Super sample each cell
+				for (int32 SubSampleY = 0; SubSampleY < ShadowSettings.StaticShadowDepthMapSuperSampleFactor; SubSampleY++)
 				{
-					// Use the maximum distance of all super samples for each cell, to get a conservative shadow map
-					MaxSampleDistance = FMath::Max(MaxSampleDistance, (Intersection.IntersectionVertex.WorldPosition - WorldSpaceStartPosition).Size3());
+					const float YFraction = (Y + SubSampleY / (float)ShadowSettings.StaticShadowDepthMapSuperSampleFactor) / (float)(ShadowDepthMap->ShadowMapSizeY - 1);
+					const float Phi = YFraction * PI;
+					const float SinPhi = FMath::Sin(Phi);
+
+					for (int32 SubSampleX = 0; SubSampleX < ShadowSettings.StaticShadowDepthMapSuperSampleFactor; SubSampleX++)
+					{
+						const float XFraction = (X + SubSampleX / (float)ShadowSettings.StaticShadowDepthMapSuperSampleFactor) / (float)(ShadowDepthMap->ShadowMapSizeX - 1);
+						const float Theta = XFraction * 2 * PI;
+						const FVector Direction(FMath::Cos(Theta) * SinPhi, FMath::Sin(Theta) * SinPhi, FMath::Cos(Phi));
+
+						const FVector4 WorldSpaceStartPosition = PointLight->Position;
+						const FVector4 WorldSpaceEndPosition = PointLight->Position + Direction * PointLight->Radius;
+						const FLightRay LightRay(
+							WorldSpaceStartPosition,
+							WorldSpaceEndPosition,
+							NULL,
+							NULL,
+							// We are tracing from the light instead of to the light,
+							// So flip sidedness so that backface culling matches up with tracing to the light
+							LIGHTRAY_FLIP_SIDEDNESS
+							);
+
+						FLightRayIntersection Intersection;
+						AggregateMesh.IntersectLightRay(LightRay, true, false, true, Context.RayCache, Intersection);
+
+						if (Intersection.bIntersects)
+						{
+							// Use the maximum distance of all super samples for each cell, to get a conservative shadow map
+							MaxSampleDistance = FMath::Max(MaxSampleDistance, (Intersection.IntersectionVertex.WorldPosition - PointLight->Position).Size3());
+						}
+					}
 				}
-			}
-		}
 
-		if (MaxSampleDistance == 0.0f)
-		{
-			MaxSampleDistance = DominantLightShadowInfo.LightSpaceImportanceBoundMax.Z - DominantLightShadowInfo.LightSpaceImportanceBoundMin.Z;
+				if (MaxSampleDistance == 0.0f)
+				{
+					MaxSampleDistance = PointLight->Radius;
+				}
+
+				ShadowDepthMap->ShadowMap[Y * ShadowDepthMap->ShadowMapSizeX + X] = FStaticShadowDepthMapSample(FFloat16(MaxSampleDistance / PointLight->Radius));
+			}
 		}
 
-		// Quantize the distance into a uint16 to reduce memory usage in Unreal
-		const uint16 QuantizedDistance = FMath::TruncToInt(65535 * FMath::Clamp(MaxSampleDistance * InvDistanceRange, 0.0f, 1.0f));
-		DominantLightShadowInfo.ShadowMap[ShadowMapY * DominantLightShadowInfo.ShadowMapSizeX + X] = FDominantLightShadowSample(QuantizedDistance);
+		FScopeLock Lock(&CompletedStaticShadowDepthMapsSync);
+		CompletedStaticShadowDepthMaps.Add(PointLight, ShadowDepthMap);
 	}
-	Context.Stats.DominantShadowThreadTime = FPlatformTime::Seconds() - CalculateWorkRangeStart;
+
+	const float NewTime = FPlatformTime::Seconds() - StartTime;
+	Context.Stats.StaticShadowDepthMapThreadTime = NewTime;
+	Context.Stats.MaxStaticShadowDepthMapThreadTime = NewTime;
 }
 
 /**
diff --git a/Engine/Source/Programs/UnrealLightmass/Private/Lighting/LightingSystem.h b/Engine/Source/Programs/UnrealLightmass/Private/Lighting/LightingSystem.h
--- a/Engine/Source/Programs/UnrealLightmass/Private/Lighting/LightingSystem.h
+++ b/Engine/Source/Programs/UnrealLightmass/Private/Lighting/LightingSystem.h
@@ -609,21 +609,21 @@ public:
 	TArray<uint8> VisibilityData;
 };
 
-/** Stores depth for a single cell of a shadow map for a dominant light. */
-class FDominantLightShadowSample : public FDominantLightShadowSampleData
+/** Stores depth for a single cell of a shadow map for a stationary light. */
+class FStaticShadowDepthMapSample : public FStaticShadowDepthMapSampleData
 {
 public:
-	FDominantLightShadowSample(uint16 InDistance)
+	FStaticShadowDepthMapSample(FFloat16 InDistance)
 	{
 		Distance = InDistance;
 	}
 };
 
 /** Stores information about how ShadowMap was generated. */
-class FDominantLightShadowInfo : public FDominantLightShadowInfoData
+class FStaticShadowDepthMap : public FStaticShadowDepthMapData
 {
 public:
-	TArray<FDominantLightShadowSample> ShadowMap;
+	TArray<FStaticShadowDepthMapSample> ShadowMap;
 };
 
 /** Number of light bounces that we are keeping track of stats for */
@@ -723,8 +723,11 @@ public:
 	/** Number of mesh queries that were trivially occluded because their group was already determined to be occluded. */
 	uint64 NumPrecomputedVisibilityMeshQueriesSkipped;
 
-	/** Thread seconds calculating the dominant shadow map. */
-	float DominantShadowThreadTime;
+	/** Thread seconds calculating static shadow depth maps. */
+	float StaticShadowDepthMapThreadTime;
+
+	/** Longest time spent on a single static shadow depth map. */
+	float MaxStaticShadowDepthMapThreadTime;
 
 	/** Thread seconds calculating the volume distance field. */
 	float VolumeDistanceFieldThreadTime;
@@ -834,7 +837,8 @@ public:
 		PrecomputedVisibilityImportanceSampleThreadTime(0),
 		NumPrecomputedVisibilityGroupQueries(0),
 		NumPrecomputedVisibilityMeshQueriesSkipped(0),
-		DominantShadowThreadTime(0),
+		StaticShadowDepthMapThreadTime(0),
+		MaxStaticShadowDepthMapThreadTime(0),
 		VolumeDistanceFieldThreadTime(0),
 		BlockOnIndirectLightingCacheTasksTime(0),
 		BlockOnIndirectLightingInterpolateTasksTime(0),
@@ -898,7 +902,8 @@ public:
 		PrecomputedVisibilityImportanceSampleThreadTime += B.PrecomputedVisibilityImportanceSampleThreadTime;
 		NumPrecomputedVisibilityGroupQueries += B.NumPrecomputedVisibilityGroupQueries;
 		NumPrecomputedVisibilityMeshQueriesSkipped += B.NumPrecomputedVisibilityMeshQueriesSkipped;
-		DominantShadowThreadTime += B.DominantShadowThreadTime;
+		StaticShadowDepthMapThreadTime += B.StaticShadowDepthMapThreadTime;
+		MaxStaticShadowDepthMapThreadTime = FMath::Max(MaxStaticShadowDepthMapThreadTime, B.MaxStaticShadowDepthMapThreadTime);
 		VolumeDistanceFieldThreadTime += B.VolumeDistanceFieldThreadTime;
 		BlockOnIndirectLightingCacheTasksTime += B.BlockOnIndirectLightingCacheTasksTime;
 		BlockOnIndirectLightingInterpolateTasksTime += B.BlockOnIndirectLightingInterpolateTasksTime;
@@ -2195,11 +2200,8 @@ private:
 	/** Calculates visibility for a given group of cells, called from all threads. */
 	void CalculatePrecomputedVisibility(int32 BucketIndex);
 
-	/** Initializes DominantLightShadowInfo and prepares for multithreaded generation of DominantLightShadowInfo.ShadowMap. */
-	void BeginCalculateDominantShadowInfo(FGuid LightGuid);
-
-	/** Generates a single row of the dominant light shadow map. */
-	void CalculateDominantShadowInfoWorkRange(int32 ShadowMapY);
+	/** Computes a shadow depth map for a stationary light. */
+	void CalculateStaticShadowDepthMap(FGuid LightGuid);
 
 	/** Prepares for multithreaded generation of VolumeDistanceField. */
 	void BeginCalculateVolumeDistanceField();
@@ -2562,23 +2564,11 @@ private:
 	/** All precomputed visibility cells in the scene.  Some of these may be processed on other agents. */
 	TArray<FPrecomputedVisibilityCell> AllPrecomputedVisibilityCells;
 
-	/** Positive if the lighting threads are done writing to DominantLightShadowInfo, and the samples can be exported. */
-	volatile int32 bDominantShadowTaskComplete;
-	/** */
-	volatile int32 NumOutstandingDominantShadowColumns;
-	/** Index of the last row of the dominant shadow map that was processed. */
-	volatile int32 OutstandingDominantShadowYIndex;
-	/** Information about the dominant light's shadow that is needed by Unreal. */
-	FDominantLightShadowInfo DominantLightShadowInfo;
-
-	/** Guid of the dominant directional light that DominantLightShadowInfo corresponds to. */
-	FGuid DominantDirectionalLightId;
-
-	/** Map from dominant spotlight to the shadowmap generated for that light. */
-	TMap<const FSpotLight*, FDominantLightShadowInfo> DominantSpotLightShadowInfos;
+	/** Threads must acquire this critical section before reading or writing to CompletedStaticShadowDepthMaps. */
+	FCriticalSection CompletedStaticShadowDepthMapsSync;
 
-	/** Threads must acquire this critical section before reading or writing to DominantSpotLightShadowInfos. */
-	FCriticalSection DominantLightShadowSync;
+	/** Static shadow depth maps ready to be exported by the main thread. */
+	TMap<const FLight*, FStaticShadowDepthMap*> CompletedStaticShadowDepthMaps;
 
 	/** Non-zero if the mesh area light data task should be exported. */
 	volatile int32 bShouldExportMeshAreaLightData;
diff --git a/Engine/Source/Programs/UnrealLightmass/Private/Lighting/TextureMapping.cpp b/Engine/Source/Programs/UnrealLightmass/Private/Lighting/TextureMapping.cpp
--- a/Engine/Source/Programs/UnrealLightmass/Private/Lighting/TextureMapping.cpp
+++ b/Engine/Source/Programs/UnrealLightmass/Private/Lighting/TextureMapping.cpp
@@ -2693,7 +2693,7 @@ void FStaticLightingSystem::CalculateIndirectLightingTextureMapping(
 	// Whether to debug the task containing the selected texel only
 	const bool bDebugSelectedTaskOnly = true;
 
-	if (GeneralSettings.NumIndirectLightingBounces > 0 || AmbientOcclusionSettings.bUseAmbientOcclusion)
+	if (GeneralSettings.NumIndirectLightingBounces > 0 || AmbientOcclusionSettings.bUseAmbientOcclusion || SkyLights.Num() > 0)
 	{
 		const double StartCacheTime = FPlatformTime::Seconds();
 
diff --git a/Engine/Source/Programs/UnrealLightmass/Public/ImportExport.h b/Engine/Source/Programs/UnrealLightmass/Public/ImportExport.h
--- a/Engine/Source/Programs/UnrealLightmass/Public/ImportExport.h
+++ b/Engine/Source/Programs/UnrealLightmass/Public/ImportExport.h
@@ -253,21 +253,18 @@ namespace Lightmass
 		float DirectionalLightShadowing;
 	};
 
-	/** Distance to the nearest shadow casting geometry for a single dominant light shadow map cell. */
-	class FDominantLightShadowSampleData
+	/** */
+	class FStaticShadowDepthMapSampleData
 	{
 	public:
-		uint16 Distance;
+		FFloat16 Distance;
 	};
 
-	class FDominantLightShadowInfoData
+	class FStaticShadowDepthMapData
 	{
 	public:
-		/** Transform from world space to the coordinate space that FDominantLightShadowSampleData's are stored in. */
+		/** Transform from world space to the coordinate space that FStaticShadowDepthMapSampleData's are stored in. */
 		FMatrix WorldToLight;
-		/** Light space bounds of the importance volume. */
-		FVector4 LightSpaceImportanceBoundMin;
-		FVector4 LightSpaceImportanceBoundMax;
 		/** Dimensions of the generated shadow map. */
 		int32 ShadowMapSizeX;
 		int32 ShadowMapSizeY;
diff --git a/Engine/Source/Programs/UnrealLightmass/Public/SceneExport.h b/Engine/Source/Programs/UnrealLightmass/Public/SceneExport.h
--- a/Engine/Source/Programs/UnrealLightmass/Public/SceneExport.h
+++ b/Engine/Source/Programs/UnrealLightmass/Public/SceneExport.h
@@ -391,14 +391,14 @@ public:
 	int32 MinDistanceFieldUpsampleFactor;
 
 	/** Distance in world space units between dominant light shadow map cells. */
-	float DominantShadowTransitionSampleDistanceX;
-	float DominantShadowTransitionSampleDistanceY;
+	float StaticShadowDepthMapTransitionSampleDistanceX;
+	float StaticShadowDepthMapTransitionSampleDistanceY;
 
 	/** Amount to super sample dominant shadow map generation, in each dimension.  Larger factors increase build time but produce a more conservative shadow map. */
-	int32 DominantShadowSuperSampleFactor;
+	int32 StaticShadowDepthMapSuperSampleFactor;
 
 	/** Maximum number of dominant shadow samples to generate for one dominant light, used to limit memory used. */
-	int32 DominantShadowMaxSamples;
+	int32 StaticShadowDepthMapMaxSamples;
 
 	/** Fraction of valid lighting samples (mapped texels or vertex samples) that must be unoccluded in a precomputed shadowmap for the shadowmap to be kept. */
 	float MinUnoccludedFraction;
diff --git a/Engine/Source/Runtime/Core/Private/Misc/Core.cpp b/Engine/Source/Runtime/Core/Private/Misc/Core.cpp
--- a/Engine/Source/Runtime/Core/Private/Misc/Core.cpp
+++ b/Engine/Source/Runtime/Core/Private/Misc/Core.cpp
@@ -238,6 +238,7 @@ DEFINE_STAT(STAT_MemoryICUTotalAllocationSize);
 DEFINE_STAT(STAT_MemoryICUDataFileAllocationSize);
 DEFINE_STAT(STAT_AnimationMemory);
 DEFINE_STAT(STAT_PrecomputedVisibilityMemory);
+DEFINE_STAT(STAT_PrecomputedShadowDepthMapMemory);
 DEFINE_STAT(STAT_PrecomputedLightVolumeMemory);
 DEFINE_STAT(STAT_StaticMeshTotalMemory);
 DEFINE_STAT(STAT_SkeletalMeshVertexMemory);
diff --git a/Engine/Source/Runtime/Core/Public/Misc/CoreStats.h b/Engine/Source/Runtime/Core/Public/Misc/CoreStats.h
--- a/Engine/Source/Runtime/Core/Public/Misc/CoreStats.h
+++ b/Engine/Source/Runtime/Core/Public/Misc/CoreStats.h
@@ -25,6 +25,7 @@ CORE_API DECLARE_LOG_CATEGORY_EXTERN(LogLoad, Log, All);
 DECLARE_MEMORY_STAT_EXTERN(TEXT("Audio Memory Used"),STAT_AudioMemory,STATGROUP_Memory, CORE_API);
 DECLARE_MEMORY_STAT_EXTERN(TEXT("Animation Memory"),STAT_AnimationMemory,STATGROUP_Memory, CORE_API);
 DECLARE_MEMORY_STAT_EXTERN(TEXT("Precomputed Visibility Memory"),STAT_PrecomputedVisibilityMemory,STATGROUP_Memory, CORE_API);
+DECLARE_MEMORY_STAT_EXTERN(TEXT("Precomputed Shadow Depth Map Memory"),STAT_PrecomputedShadowDepthMapMemory,STATGROUP_Memory, CORE_API);
 DECLARE_MEMORY_STAT_EXTERN(TEXT("Precomputed Light Volume Memory"),STAT_PrecomputedLightVolumeMemory,STATGROUP_Memory, CORE_API);
 DECLARE_MEMORY_STAT_EXTERN(TEXT("StaticMesh Total Memory"),STAT_StaticMeshTotalMemory,STATGROUP_Memory, CORE_API);
 DECLARE_MEMORY_STAT_EXTERN(TEXT("SkeletalMesh Vertex Memory"),STAT_SkeletalMeshVertexMemory,STATGROUP_Memory, CORE_API);
diff --git a/Engine/Source/Runtime/Core/Public/UObject/ObjectVersion.h b/Engine/Source/Runtime/Core/Public/UObject/ObjectVersion.h
--- a/Engine/Source/Runtime/Core/Public/UObject/ObjectVersion.h
+++ b/Engine/Source/Runtime/Core/Public/UObject/ObjectVersion.h
@@ -568,6 +568,8 @@ enum EUnrealEngineObjectUE4Version
 	VER_UE4_GAMEPLAY_TAG_CONTAINER_TAG_TYPE_CHANGE,
 	// Change from UInstancedFoliageSettings to UFoliageType, and change the api from being keyed on UStaticMesh* to UFoliageType*
 	VER_UE4_FOLIAGE_SETTINGS_TYPE,
+	// Lights serialize static shadow depth maps
+	VER_UE4_STATIC_SHADOW_DEPTH_MAPS,
 
 	// -----<new versions can be added before this line>-------------------------------------------------
 	// - this needs to be the last line (see note below)
diff --git a/Engine/Source/Runtime/Engine/Classes/Components/LightComponent.h b/Engine/Source/Runtime/Engine/Classes/Components/LightComponent.h
--- a/Engine/Source/Runtime/Engine/Classes/Components/LightComponent.h
+++ b/Engine/Source/Runtime/Engine/Classes/Components/LightComponent.h
@@ -7,6 +7,34 @@
 
 #include "LightComponent.generated.h"
 
+/** 
+ * A texture containing depth values of static objects that was computed during the lighting build.
+ * Used by Stationary lights to shadow translucency.
+ */
+class FStaticShadowDepthMap : public FTexture
+{
+public:
+	/** Transform from world space to the coordinate space that DepthSamples are stored in. */
+	FMatrix WorldToLight;
+	/** Dimensions of the generated shadow map. */
+	int32 ShadowMapSizeX;
+	int32 ShadowMapSizeY;
+	/** Shadowmap depth values */
+	TArray<FFloat16> DepthSamples;
+
+	FStaticShadowDepthMap() :
+		WorldToLight(FMatrix::Identity),
+		ShadowMapSizeX(0),
+		ShadowMapSizeY(0)
+	{}
+
+	virtual void InitRHI();
+
+	void Empty();
+
+	friend FArchive& operator<<(FArchive& Ar, FStaticShadowDepthMap& ShadowMap);
+};
+
 UCLASS(abstract, HideCategories=(Trigger,Activation,"Components|Activation",Physics), ShowCategories=(Mobility))
 class ENGINE_API ULightComponent : public ULightComponentBase
 {
@@ -163,6 +191,11 @@ public:
 	/** The light's scene info. */
 	class FLightSceneProxy* SceneProxy;
 
+	FStaticShadowDepthMap StaticShadowDepthMap;
+
+	/** Fence used to track progress of render resource destruction. */
+	FRenderCommandFence DestroyFence;
+
 	/**
 	 * Test whether this light affects the given primitive.  This checks both the primitive and light settings for light relevance
 	 * and also calls AffectsBounds.
@@ -222,6 +255,7 @@ public:
 	float ComputeLightBrightness() const;
 
 	// Begin UObject interface.
+	virtual void Serialize(FArchive& Ar) override;
 	virtual void PostLoad() override;
 #if WITH_EDITOR
 	virtual void PreEditUndo() override;
@@ -230,6 +264,8 @@ public:
 	virtual void PostEditChangeProperty(FPropertyChangedEvent& PropertyChangedEvent) override;
 	virtual void UpdateLightSpriteTexture() override;
 #endif // WITH_EDITOR
+	virtual void BeginDestroy() override;
+	virtual bool IsReadyForFinishDestroy() override;
 	// End UObject interface.
 
 	virtual TSharedPtr<FComponentInstanceDataBase> GetComponentInstanceData() const override;
diff --git a/Engine/Source/Runtime/Engine/Private/LightComponent.cpp b/Engine/Source/Runtime/Engine/Private/LightComponent.cpp
--- a/Engine/Source/Runtime/Engine/Private/LightComponent.cpp
+++ b/Engine/Source/Runtime/Engine/Private/LightComponent.cpp
@@ -11,6 +11,52 @@
 #include "MessageLog.h"
 #include "UObjectToken.h"
 #include "ComponentInstanceDataCache.h"
+#include "TargetPlatform.h"
+
+void FStaticShadowDepthMap::InitRHI()
+{
+	if (ShadowMapSizeX > 0 && ShadowMapSizeY > 0 && GRHIFeatureLevel >= ERHIFeatureLevel::SM3)
+	{
+		FRHIResourceCreateInfo CreateInfo;
+		FTexture2DRHIRef Texture2DRHI = RHICreateTexture2D(ShadowMapSizeX, ShadowMapSizeY, PF_R16F, 1, 1, 0, CreateInfo);
+		TextureRHI = Texture2DRHI;
+
+		uint32 DestStride = 0;
+		uint8* TextureData = (uint8*)RHILockTexture2D(Texture2DRHI, 0, RLM_WriteOnly, DestStride, false);
+		uint32 RowSize = ShadowMapSizeX * GPixelFormats[PF_R16F].BlockBytes;
+
+		for (int32 Y = 0; Y < ShadowMapSizeY; Y++)
+		{
+			FMemory::Memcpy(TextureData + DestStride * Y, ((uint8*)DepthSamples.GetData()) + RowSize * Y, RowSize);
+		}
+
+		RHIUnlockTexture2D(Texture2DRHI, 0, false);
+	}
+}
+
+void FStaticShadowDepthMap::Empty()
+{
+	DEC_DWORD_STAT_BY(STAT_PrecomputedShadowDepthMapMemory, DepthSamples.GetAllocatedSize());
+
+	ShadowMapSizeX = 0;
+	ShadowMapSizeY = 0;
+	DepthSamples.Empty();
+}
+
+FArchive& operator<<(FArchive& Ar, FStaticShadowDepthMap& ShadowMap)
+{
+	Ar << ShadowMap.WorldToLight;
+	Ar << ShadowMap.ShadowMapSizeX;
+	Ar << ShadowMap.ShadowMapSizeY;
+	Ar << ShadowMap.DepthSamples;
+
+	if (Ar.IsLoading())
+	{
+		INC_DWORD_STAT_BY(STAT_PrecomputedShadowDepthMapMemory, ShadowMap.DepthSamples.GetAllocatedSize());
+	}
+
+	return Ar;
+}
 
 void ULightComponentBase::Serialize(FArchive& Ar)
 {
@@ -192,6 +238,8 @@ FLightSceneProxy::FLightSceneProxy(const ULightComponent* InLightComponent)
 	LightFunctionScale = LightComponent->LightFunctionScale;
 	LightFunctionFadeDistance = LightComponent->LightFunctionFadeDistance;
 	LightFunctionDisabledBrightness = LightComponent->DisabledBrightness;
+
+	StaticShadowDepthMap = &LightComponent->StaticShadowDepthMap;
 }
 
 bool FLightSceneProxy::ShouldCreatePerObjectShadowsForDynamicObjects() const
@@ -313,6 +361,25 @@ float ULightComponent::ComputeLightBrightness() const
 	return LightBrightness;
 }
 
+void ULightComponent::Serialize(FArchive& Ar)
+{
+	Super::Serialize(Ar);
+
+	if (Ar.UE4Ver() >= VER_UE4_STATIC_SHADOW_DEPTH_MAPS)
+	{
+		if (Ar.IsCooking() && !Ar.CookingTarget()->SupportsFeature(ETargetPlatformFeatures::HighQualityLightmaps))
+		{
+			// Toss lighting data only needed for high quality lightmaps
+			FStaticShadowDepthMap EmptyDepthMap;
+			Ar << EmptyDepthMap;
+		}
+		else
+		{
+			Ar << StaticShadowDepthMap;
+		}
+	}
+}
+
 /**
  * Called after this UObject has been serialized
  */
@@ -348,6 +415,11 @@ void ULightComponent::PostLoad()
 			}
 		}
 	}
+
+	if (HasStaticShadowing() && !HasStaticLighting())
+	{
+		BeginInitResource(&StaticShadowDepthMap);
+	}
 }
 
 #if WITH_EDITOR
@@ -488,6 +560,26 @@ void ULightComponent::UpdateLightSpriteTexture()
 
 #endif // WITH_EDITOR
 
+void ULightComponent::BeginDestroy()
+{
+	Super::BeginDestroy();
+
+	if (HasStaticShadowing() && !HasStaticLighting())
+	{
+		BeginReleaseResource(&StaticShadowDepthMap);
+		StaticShadowDepthMap.Empty();
+	}
+
+	// Use a fence to keep track of when the rendering thread executes the release command
+	DestroyFence.BeginFence();
+}
+
+bool ULightComponent::IsReadyForFinishDestroy()
+{
+	// Don't allow the light component to be destroyed until its rendering resources have been released
+	return Super::IsReadyForFinishDestroy() && DestroyFence.IsFenceComplete();
+}
+
 void ULightComponent::OnRegister()
 {
 	Super::OnRegister();
@@ -731,6 +823,15 @@ void ULightComponent::InvalidateLightingCacheInner(bool bRecreateLightGuids)
 	// Save the light state for transactions.
 	Modify();
 
+	// Detach the component from the scene for the duration of this function.
+	FComponentReregisterContext ReregisterContext(this);
+
+	// Block until the RT processes the unregister before modifying variables that it may need to access
+	FlushRenderingCommands();
+
+	StaticShadowDepthMap.Empty();
+	BeginReleaseResource(&StaticShadowDepthMap);
+
 	bPrecomputedLightingIsValid = false;
 
 	if (bRecreateLightGuids)
diff --git a/Engine/Source/Runtime/Engine/Private/SpotLightComponent.cpp b/Engine/Source/Runtime/Engine/Private/SpotLightComponent.cpp
--- a/Engine/Source/Runtime/Engine/Private/SpotLightComponent.cpp
+++ b/Engine/Source/Runtime/Engine/Private/SpotLightComponent.cpp
@@ -127,6 +127,7 @@ public:
 		OutInitializer.WorldToLight = GetWorldToLight().RemoveTranslation();
 		OutInitializer.Scales = FVector(1.0f,InvTanOuterCone,InvTanOuterCone);
 		OutInitializer.FaceDirection = FVector(1,0,0);
+		//@todo - these half radius bounds are not correct, but quite a few things depend on it at this point
 		OutInitializer.SubjectBounds = FBoxSphereBounds(
 			GetLightToWorld().RemoveTranslation().TransformPosition(FVector(Radius / 2.0f,0,0)),
 			FVector(Radius/2.0f,Radius/2.0f,Radius/2.0f),
diff --git a/Engine/Source/Runtime/Engine/Public/SceneManagement.h b/Engine/Source/Runtime/Engine/Public/SceneManagement.h
--- a/Engine/Source/Runtime/Engine/Public/SceneManagement.h
+++ b/Engine/Source/Runtime/Engine/Public/SceneManagement.h
@@ -741,6 +741,7 @@ public:
 
 	inline bool HasReflectiveShadowMap() const { return bHasReflectiveShadowMap; }
 	inline bool NeedsLPVInjection() const { return bAffectDynamicIndirectLighting; }
+	inline const class FStaticShadowDepthMap* GetStaticShadowDepthMap() const { return StaticShadowDepthMap; }
 
 	/**
 	 * Shifts light position and all relevant data by an arbitrary delta.
@@ -795,6 +796,8 @@ protected:
 	/** Transient shadowmap channel used to preview the results of stationary light shadowmap packing. */
 	int32 PreviewShadowMapChannel;
 
+	const class FStaticShadowDepthMap* StaticShadowDepthMap;
+
 	/** Light function parameters. */
 	FVector	LightFunctionScale;
 	float LightFunctionFadeDistance;
diff --git a/Engine/Source/Runtime/Renderer/Private/TranslucentLighting.cpp b/Engine/Source/Runtime/Renderer/Private/TranslucentLighting.cpp
--- a/Engine/Source/Runtime/Renderer/Private/TranslucentLighting.cpp
+++ b/Engine/Source/Runtime/Renderer/Private/TranslucentLighting.cpp
@@ -710,11 +710,11 @@ public:
 		const FLightSceneInfo* LightSceneInfo, 
 		const FProjectedShadowInfo* ShadowMap, 
 		int32 VolumeCascadeIndexValue, 
-		bool bShadowed) const
+		bool bDynamicallyShadowed) const
 	{
 		SetDeferredLightParameters(RHICmdList, ShaderRHI, Shader->GetUniformBufferParameter<FDeferredLightUniformStruct>(), LightSceneInfo, View);
 
-		if (bShadowed)
+		if (bDynamicallyShadowed)
 		{
 			FVector4 ShadowmapMinMaxValue;
 			FMatrix WorldToShadowMatrixValue = ShadowMap->GetWorldToShadowMatrix(ShadowmapMinMaxValue);
@@ -790,7 +790,7 @@ private:
 IMPLEMENT_SHADER_TYPE(,FTranslucentObjectShadowingPS,TEXT("TranslucentLightingShaders"),TEXT("PerObjectShadowingMainPS"),SF_Pixel);
 
 /** Shader that adds direct lighting contribution from the given light to the current volume lighting cascade. */
-template<ELightComponentType InjectionType, bool bShadowed, bool bApplyLightFunction, bool bInverseSquared>
+template<ELightComponentType InjectionType, bool bDynamicallyShadowed, bool bApplyLightFunction, bool bInverseSquared>
 class TTranslucentLightingInjectPS : public FMaterialShader
 {
 	DECLARE_SHADER_TYPE(TTranslucentLightingInjectPS,Material);
@@ -801,7 +801,7 @@ public:
 		FMaterialShader::ModifyCompilationEnvironment(Platform, Material, OutEnvironment);
 		OutEnvironment.SetDefine(TEXT("RADIAL_ATTENUATION"), (uint32)(InjectionType != LightType_Directional));
 		OutEnvironment.SetDefine(TEXT("INJECTION_PIXEL_SHADER"), 1);
-		OutEnvironment.SetDefine(TEXT("DYNAMICALLY_SHADOWED"), (uint32)bShadowed);
+		OutEnvironment.SetDefine(TEXT("DYNAMICALLY_SHADOWED"), (uint32)bDynamicallyShadowed);
 		OutEnvironment.SetDefine(TEXT("APPLY_LIGHT_FUNCTION"), (uint32)bApplyLightFunction);
 		OutEnvironment.SetDefine(TEXT("INVERSE_SQUARED_FALLOFF"), (uint32)bInverseSquared);
 	}
@@ -831,6 +831,10 @@ public:
 		LightFunctionParameters.Bind(Initializer.ParameterMap);
 		TranslucentInjectParameters.Bind(Initializer.ParameterMap);
 		LightFunctionWorldToLight.Bind(Initializer.ParameterMap, TEXT("LightFunctionWorldToLight"));
+		bStaticallyShadowed.Bind(Initializer.ParameterMap, TEXT("bStaticallyShadowed"));
+		StaticShadowDepthTexture.Bind(Initializer.ParameterMap, TEXT("StaticShadowDepthTexture"));
+		StaticShadowDepthTextureSampler.Bind(Initializer.ParameterMap, TEXT("StaticShadowDepthTextureSampler"));
+		WorldToStaticShadowMatrix.Bind(Initializer.ParameterMap, TEXT("WorldToStaticShadowMatrix"));
 	}
 	TTranslucentLightingInjectPS() {}
 
@@ -845,7 +849,7 @@ public:
 		int32 InnerSplitIndex, 
 		int32 VolumeCascadeIndexValue)
 	{
-		check(ShadowMap || !bShadowed);
+		check(ShadowMap || !bDynamicallyShadowed);
 		
 		const FPixelShaderRHIParamRef ShaderRHI = GetPixelShader();
 
@@ -898,7 +902,7 @@ public:
 		//@todo - needs to be a permutation to reduce shadow filtering work
 		SetShaderValue(RHICmdList, ShaderRHI, SpotlightMask, (LightSceneInfo->Proxy->GetLightType() == LightType_Spot ? 1.0f : 0.0f));
 
-		if (bShadowed)
+		if (bDynamicallyShadowed)
 		{
 			SetShaderValue(RHICmdList, ShaderRHI, DepthBiasParameters, FVector2D(ShadowMap->GetShaderDepthBias(), 1.0f / (ShadowMap->MaxSubjectZ - ShadowMap->MinSubjectZ)));
 
@@ -912,13 +916,13 @@ public:
 				);
 		}
 
-		if (bShadowed && InjectionType == LightType_Point)
+		if (bDynamicallyShadowed && InjectionType == LightType_Point)
 		{
 			OnePassShadowParameters.Set(RHICmdList, ShaderRHI, ShadowMap);
 		}
 
 		LightFunctionParameters.Set(RHICmdList, ShaderRHI, LightSceneInfo, 1);
-		TranslucentInjectParameters.Set(RHICmdList, ShaderRHI, this, View, LightSceneInfo, ShadowMap, VolumeCascadeIndexValue, bShadowed);
+		TranslucentInjectParameters.Set(RHICmdList, ShaderRHI, this, View, LightSceneInfo, ShadowMap, VolumeCascadeIndexValue, bDynamicallyShadowed);
 
 		if (LightFunctionWorldToLight.IsBound())
 		{
@@ -929,6 +933,24 @@ public:
 
 			SetShaderValue(RHICmdList, ShaderRHI, LightFunctionWorldToLight, WorldToLight);
 		}
+
+		const FStaticShadowDepthMap* StaticShadowDepthMap = LightSceneInfo->Proxy->GetStaticShadowDepthMap();
+		const uint32 bStaticallyShadowedValue = LightSceneInfo->bPrecomputedLightingIsValid && StaticShadowDepthMap && StaticShadowDepthMap->TextureRHI ? 1 : 0;
+		FTextureRHIParamRef StaticShadowDepthMapTexture = StaticShadowDepthMap ? StaticShadowDepthMap->TextureRHI : GWhiteTexture->TextureRHI;
+		const FMatrix WorldToStaticShadow = StaticShadowDepthMap ? StaticShadowDepthMap->WorldToLight : FMatrix::Identity;
+
+		SetShaderValue(RHICmdList, ShaderRHI, bStaticallyShadowed, bStaticallyShadowedValue);
+
+		SetTextureParameter(
+			RHICmdList, 
+			ShaderRHI,
+			StaticShadowDepthTexture,
+			StaticShadowDepthTextureSampler,
+			TStaticSamplerState<SF_Point,AM_Clamp,AM_Clamp,AM_Clamp>::GetRHI(),
+			StaticShadowDepthMapTexture
+			);
+
+		SetShaderValue(RHICmdList, ShaderRHI, WorldToStaticShadowMatrix, WorldToStaticShadow);
 	}
 
 	virtual bool Serialize(FArchive& Ar)
@@ -946,6 +968,10 @@ public:
 		Ar << LightFunctionParameters;
 		Ar << TranslucentInjectParameters;
 		Ar << LightFunctionWorldToLight;
+		Ar << bStaticallyShadowed;
+		Ar << StaticShadowDepthTexture;
+		Ar << StaticShadowDepthTextureSampler;
+		Ar << WorldToStaticShadowMatrix;
 		return bShaderHasOutdatedParameters;
 	}
 
@@ -962,11 +988,15 @@ private:
 	FLightFunctionSharedParameters LightFunctionParameters;
 	FTranslucentInjectParameters TranslucentInjectParameters;
 	FShaderParameter LightFunctionWorldToLight;
+	FShaderParameter bStaticallyShadowed;
+	FShaderResourceParameter StaticShadowDepthTexture;
+	FShaderResourceParameter StaticShadowDepthTextureSampler;
+	FShaderParameter WorldToStaticShadowMatrix;
 };
 
-#define IMPLEMENT_INJECTION_PIXELSHADER_TYPE(LightType,bShadowed,bApplyLightFunction,bInverseSquared) \
-	typedef TTranslucentLightingInjectPS<LightType,bShadowed,bApplyLightFunction,bInverseSquared> TTranslucentLightingInjectPS##LightType##bShadowed##bApplyLightFunction##bInverseSquared; \
-	IMPLEMENT_MATERIAL_SHADER_TYPE(template<>,TTranslucentLightingInjectPS##LightType##bShadowed##bApplyLightFunction##bInverseSquared,TEXT("TranslucentLightInjectionShaders"),TEXT("InjectMainPS"),SF_Pixel);
+#define IMPLEMENT_INJECTION_PIXELSHADER_TYPE(LightType,bDynamicallyShadowed,bApplyLightFunction,bInverseSquared) \
+	typedef TTranslucentLightingInjectPS<LightType,bDynamicallyShadowed,bApplyLightFunction,bInverseSquared> TTranslucentLightingInjectPS##LightType##bDynamicallyShadowed##bApplyLightFunction##bInverseSquared; \
+	IMPLEMENT_MATERIAL_SHADER_TYPE(template<>,TTranslucentLightingInjectPS##LightType##bDynamicallyShadowed##bApplyLightFunction##bInverseSquared,TEXT("TranslucentLightInjectionShaders"),TEXT("InjectMainPS"),SF_Pixel);
 
 /** Versions with a light function. */
 IMPLEMENT_INJECTION_PIXELSHADER_TYPE(LightType_Directional,true,true,false); 
@@ -1466,7 +1496,7 @@ void FDeferredShadingSceneRenderer::AccumulateTranslucentVolumeObjectShadowing(F
  * @param MaterialProxy must not be 0
  * @param InnerSplitIndex todo: get from ShadowMap, INDEX_NONE if no directional light
  */
-template<ELightComponentType InjectionType, bool bShadowed>
+template<ELightComponentType InjectionType, bool bDynamicallyShadowed>
 void SetInjectionShader(
 	FRHICommandList& RHICmdList,
 	const FViewInfo& View, 
@@ -1480,7 +1510,7 @@ void SetInjectionShader(
 	bool bApplyLightFunction,
 	bool bInverseSquared)
 {
-	check(ShadowMap || !bShadowed);
+	check(ShadowMap || !bDynamicallyShadowed);
 
 	const FMaterialShaderMap* MaterialShaderMap = MaterialProxy->GetMaterial(View.GetFeatureLevel())->GetRenderingThreadShaderMap();
 	FMaterialShader* PixelShader = NULL;
@@ -1491,14 +1521,14 @@ void SetInjectionShader(
 	{
 		if( bInverseSquared )
 		{
-			auto InjectionPixelShader = MaterialShaderMap->GetShader< TTranslucentLightingInjectPS<InjectionType, bShadowed, true, true && !Directional> >();
+			auto InjectionPixelShader = MaterialShaderMap->GetShader< TTranslucentLightingInjectPS<InjectionType, bDynamicallyShadowed, true, true && !Directional> >();
 
 			check(InjectionPixelShader);
 			PixelShader = InjectionPixelShader;
 		}
 		else
 		{
-			auto InjectionPixelShader = MaterialShaderMap->GetShader< TTranslucentLightingInjectPS<InjectionType, bShadowed, true, false> >();
+			auto InjectionPixelShader = MaterialShaderMap->GetShader< TTranslucentLightingInjectPS<InjectionType, bDynamicallyShadowed, true, false> >();
 
 			check(InjectionPixelShader);
 			PixelShader = InjectionPixelShader;
@@ -1508,22 +1538,22 @@ void SetInjectionShader(
 	{
 		if( bInverseSquared )
 		{
-			auto InjectionPixelShader = MaterialShaderMap->GetShader< TTranslucentLightingInjectPS<InjectionType, bShadowed, false, true && !Directional> >();
+			auto InjectionPixelShader = MaterialShaderMap->GetShader< TTranslucentLightingInjectPS<InjectionType, bDynamicallyShadowed, false, true && !Directional> >();
 
 			check(InjectionPixelShader);
 			PixelShader = InjectionPixelShader;
 		}
 		else
 		{
-			auto InjectionPixelShader = MaterialShaderMap->GetShader< TTranslucentLightingInjectPS<InjectionType, bShadowed, false, false> >();
+			auto InjectionPixelShader = MaterialShaderMap->GetShader< TTranslucentLightingInjectPS<InjectionType, bDynamicallyShadowed, false, false> >();
 
 			check(InjectionPixelShader);
 			PixelShader = InjectionPixelShader;
 		}
 	}
 	
-	FBoundShaderStateRHIRef& BoundShaderState = LightSceneInfo->TranslucentInjectBoundShaderState[InjectionType][bShadowed][bApplyLightFunction][bInverseSquared];
-	const FMaterialShaderMap*& CachedShaderMap = LightSceneInfo->TranslucentInjectCachedShaderMaps[InjectionType][bShadowed][bApplyLightFunction][bInverseSquared];
+	FBoundShaderStateRHIRef& BoundShaderState = LightSceneInfo->TranslucentInjectBoundShaderState[InjectionType][bDynamicallyShadowed][bApplyLightFunction][bInverseSquared];
+	const FMaterialShaderMap*& CachedShaderMap = LightSceneInfo->TranslucentInjectCachedShaderMaps[InjectionType][bDynamicallyShadowed][bApplyLightFunction][bInverseSquared];
 
 	// Recreate the bound shader state if the shader map has changed since the last cache
 	// This can happen due to async shader compiling
@@ -1542,13 +1572,13 @@ void SetInjectionShader(
 	{
 		if( bInverseSquared )
 		{
-			auto InjectionPixelShader = MaterialShaderMap->GetShader< TTranslucentLightingInjectPS<InjectionType, bShadowed, true, true && !Directional> >();
+			auto InjectionPixelShader = MaterialShaderMap->GetShader< TTranslucentLightingInjectPS<InjectionType, bDynamicallyShadowed, true, true && !Directional> >();
 			check(InjectionPixelShader);
 			InjectionPixelShader->SetParameters(RHICmdList, View, LightSceneInfo, MaterialProxy, ShadowMap, InnerSplitIndex, VolumeCascadeIndexValue);
 		}
 		else
 		{
-			auto InjectionPixelShader = MaterialShaderMap->GetShader< TTranslucentLightingInjectPS<InjectionType, bShadowed, true, false> >();
+			auto InjectionPixelShader = MaterialShaderMap->GetShader< TTranslucentLightingInjectPS<InjectionType, bDynamicallyShadowed, true, false> >();
 			check(InjectionPixelShader);
 			InjectionPixelShader->SetParameters(RHICmdList, View, LightSceneInfo, MaterialProxy, ShadowMap, InnerSplitIndex, VolumeCascadeIndexValue);
 		}
@@ -1557,13 +1587,13 @@ void SetInjectionShader(
 	{
 		if( bInverseSquared )
 		{
-			auto InjectionPixelShader = MaterialShaderMap->GetShader< TTranslucentLightingInjectPS<InjectionType, bShadowed, false, true && !Directional> >();
+			auto InjectionPixelShader = MaterialShaderMap->GetShader< TTranslucentLightingInjectPS<InjectionType, bDynamicallyShadowed, false, true && !Directional> >();
 			check(InjectionPixelShader);
 			InjectionPixelShader->SetParameters(RHICmdList, View, LightSceneInfo, MaterialProxy, ShadowMap, InnerSplitIndex, VolumeCascadeIndexValue);
 		}
 		else
 		{
-			auto InjectionPixelShader = MaterialShaderMap->GetShader< TTranslucentLightingInjectPS<InjectionType, bShadowed, false, false> >();
+			auto InjectionPixelShader = MaterialShaderMap->GetShader< TTranslucentLightingInjectPS<InjectionType, bDynamicallyShadowed, false, false> >();
 			check(InjectionPixelShader);
 			InjectionPixelShader->SetParameters(RHICmdList, View, LightSceneInfo, MaterialProxy, ShadowMap, InnerSplitIndex, VolumeCascadeIndexValue);
 		}
