Added tube light support to lightmass.

Includes pull request 142.

diff --git a/Engine/Source/Editor/UnrealEd/Private/Lightmass/Lightmass.cpp b/Engine/Source/Editor/UnrealEd/Private/Lightmass/Lightmass.cpp
--- a/Engine/Source/Editor/UnrealEd/Private/Lightmass/Lightmass.cpp
+++ b/Engine/Source/Editor/UnrealEd/Private/Lightmass/Lightmass.cpp
@@ -896,6 +896,7 @@ void FLightmassExporter::WriteLights( int32 Channel )
 		LightData.IndirectLightingSaturation = Light->LightmassSettings.IndirectLightingSaturation;
 		LightData.ShadowExponent = Light->LightmassSettings.ShadowExponent;
 		LightData.LightSourceRadius = 0;
+		LightData.LightSourceLength = 0;
 		DirectionalData.LightSourceAngle = Light->LightmassSettings.LightSourceAngle * (float)PI / 180.0f;
 		Swarm.WriteChannel( Channel, &LightData, sizeof(LightData) );
 		Swarm.WriteChannel( Channel, &DirectionalData, sizeof(DirectionalData) );
@@ -912,6 +913,7 @@ void FLightmassExporter::WriteLights( int32 Channel )
 		LightData.IndirectLightingSaturation = Light->LightmassSettings.IndirectLightingSaturation;
 		LightData.ShadowExponent = Light->LightmassSettings.ShadowExponent;
 		LightData.LightSourceRadius = Light->SourceRadius;
+		LightData.LightSourceLength = Light->SourceLength;
 		PointData.Radius = Light->AttenuationRadius;
 		PointData.FalloffExponent = Light->LightFalloffExponent;
 		Swarm.WriteChannel( Channel, &LightData, sizeof(LightData) );
@@ -930,6 +932,7 @@ void FLightmassExporter::WriteLights( int32 Channel )
 		LightData.IndirectLightingSaturation = Light->LightmassSettings.IndirectLightingSaturation;
 		LightData.ShadowExponent = Light->LightmassSettings.ShadowExponent;
 		LightData.LightSourceRadius = Light->SourceRadius;
+		LightData.LightSourceLength = Light->SourceLength;
 		PointData.Radius = Light->AttenuationRadius;
 		PointData.FalloffExponent = Light->LightFalloffExponent;
 		SpotData.InnerConeAngle = Light->InnerConeAngle; 
@@ -4034,4 +4037,4 @@ bool FLightmassProcessor::ImportTextureMapping(int32 Channel, FTextureMappingImp
 	return bResult;
 }
 
-#undef LOCTEXT_NAMESPACE
+#undef LOCTEXT_NAMESPACE
\ No newline at end of file
diff --git a/Engine/Source/Programs/UnrealLightmass/Private/ImportExport/LightmassScene.cpp b/Engine/Source/Programs/UnrealLightmass/Private/ImportExport/LightmassScene.cpp
--- a/Engine/Source/Programs/UnrealLightmass/Private/ImportExport/LightmassScene.cpp
+++ b/Engine/Source/Programs/UnrealLightmass/Private/ImportExport/LightmassScene.cpp
@@ -822,7 +822,7 @@ void FDirectionalLight::ValidateSurfaceSample(const FVector4& Point, FLightSurfa
 }
 
 /** Gets a single position which represents the center of the area light source from the ReceivingPosition's point of view. */
-FVector4 FDirectionalLight::LightCenterPosition(const FVector4& ReceivingPosition) const
+FVector4 FDirectionalLight::LightCenterPosition(const FVector4& ReceivingPosition, const FVector4& ReceivingNormal) const
 {
 	return ReceivingPosition - Direction * 2.0f * SceneBounds.SphereRadius;
 }
@@ -903,19 +903,41 @@ bool FPointLight::AffectsBounds(const FBoxSphereBounds& Bounds) const
  */
 FLinearColor FPointLight::GetDirectIntensity(const FVector4& Point, bool bCalculateForIndirectLighting) const
 {
-	if( LightFlags & GI_LIGHT_INVERSE_SQUARED )
+	if (LightFlags & GI_LIGHT_INVERSE_SQUARED)
 	{
-		float DistanceSqr = ( Position - Point ).SizeSquared3();
-		float DistanceAttenuation = 16.0f / ( DistanceSqr + 0.0001f );
+		FVector4 ToLight = Position - Point;
+		float DistanceSqr = ToLight.SizeSquared3();
 
-		float LightRadiusMask = FMath::Square( FMath::Max( 0.0f, 1.0f - FMath::Square( DistanceSqr / (Radius * Radius) ) ) );
+		float DistanceAttenuation = 0.0f;
+		if( LightSourceLength > 0.0f )
+		{
+			// Line segment irradiance
+			FVector4 L01 = Direction * LightSourceLength;
+			FVector4 L0 = ToLight - 0.5f * L01;
+			FVector4 L1 = ToLight + 0.5f * L01;
+			float LengthL0 = L0.Size3();
+			float LengthL1 = L1.Size3();
+
+			DistanceAttenuation = 1.0f / ( ( LengthL0 * LengthL1 + Dot3( L0, L1 ) ) * 0.5f + 1.0f );
+			DistanceAttenuation *= 0.5 * ( L0 / LengthL0 + L1 / LengthL1 ).Size3();
+		}
+		else
+		{
+			// Sphere irradiance (technically just 1/d^2 but this avoids inf)
+			DistanceAttenuation = 1.0f / ( DistanceSqr + 1.0f );
+		}
+
+		// lumens
+		DistanceAttenuation *= 16.0f;
+
+		float LightRadiusMask = FMath::Square(FMath::Max(0.0f, 1.0f - FMath::Square(DistanceSqr / (Radius * Radius))));
 		DistanceAttenuation *= LightRadiusMask;
 
 		return FLight::GetDirectIntensity(Point, bCalculateForIndirectLighting) * DistanceAttenuation;
 	}
 	else
 	{
-		float RadialAttenuation = FMath::Pow( FMath::Max(1.0f - ((Position - Point) / Radius).SizeSquared3(),0.0f), FalloffExponent );
+		float RadialAttenuation = FMath::Pow(FMath::Max(1.0f - ((Position - Point) / Radius).SizeSquared3(), 0.0f), FalloffExponent);
 
 		return FLight::GetDirectIntensity(Point, bCalculateForIndirectLighting) * RadialAttenuation;
 	}
@@ -1103,6 +1125,30 @@ float FPointLight::Power() const
 	return FLinearColorUtils::LinearRGBToXYZ(LightPower).G;
 }
 
+FVector4 FPointLight::LightCenterPosition(const FVector4& ReceivingPosition, const FVector4& ReceivingNormal) const
+{
+	if( LightSourceLength > 0 )
+	{
+		FVector4 ToLight = Position - ReceivingPosition;
+
+		FVector4 Dir = Direction;
+		if( Dot3( ReceivingNormal, Direction ) < 0.0f )
+		{
+			Dir = -Direction;
+		}
+
+		// Clip to hemisphere
+		float Proj = FMath::Min( Dot3( ToLight, Dir ), Dot3( ReceivingNormal, ToLight) / Dot3( ReceivingNormal, Dir ) );
+
+		// Point on line segment closest to Point
+		return Position - Dir * FMath::Clamp( Proj, -0.5f * LightSourceLength, 0.5f * LightSourceLength );
+	}
+	else
+	{
+		return Position;
+	}
+}
+
 /** Returns true if all parts of the light are behind the surface being tested. */
 bool FPointLight::BehindSurface(const FVector4& TrianglePoint, const FVector4& TriangleNormal) const
 {
@@ -1113,22 +1159,96 @@ bool FPointLight::BehindSurface(const FVector4& TrianglePoint, const FVector4& T
 /** Gets a single direction to use for direct lighting that is representative of the whole area light. */
 FVector4 FPointLight::GetDirectLightingDirection(const FVector4& Point, const FVector4& PointNormal) const
 {
-	// The position on the point light surface sphere that will first be visible to a triangle rotating toward the light
-	const FVector4 FirstVisibleLightPoint = Position + PointNormal * LightSourceRadius;
-	return FirstVisibleLightPoint - Point;
+	FVector4 LightPosition = Position;
+
+	if( LightSourceLength > 0 )
+	{
+		FVector4 ToLight = Position - Point;
+		FVector4 L01 = Direction * LightSourceLength;
+		FVector4 L0 = ToLight - 0.5 * L01;
+		FVector4 L1 = ToLight + 0.5 * L01;
+#if 0
+		// Point on line segment with smallest angle to normal
+		float A = LightSourceLength * LightSourceLength;
+		float B = 2.0f * Dot3( L0, L01 );
+		float C = Dot3( L0, L0 );
+		float D = Dot3( PointNormal, L0 );
+		float E = Dot3( PointNormal, L01 );
+		float t = FMath::Clamp( (B*D - 2.0f * C*E) / (B*E - 2.0f * A*D), 0.0f, 1.0f );
+		return L0 + t * L01;
+#else
+		// Line segment irradiance
+		float LengthL0 = L0.Size3();
+		float LengthL1 = L1.Size3();
+		return ( L0 * LengthL1 + L1 * LengthL0 ) / ( LengthL0 + LengthL1 );
+#endif
+	}
+	else
+	{
+		// The position on the point light surface sphere that will first be visible to a triangle rotating toward the light
+		const FVector4 FirstVisibleLightPoint = LightPosition + PointNormal * LightSourceRadius;
+		return FirstVisibleLightPoint - Point;
+	}
 }
 
 /** Generates a sample on the light's surface. */
 void FPointLight::SampleLightSurface(FLMRandomStream& RandomStream, FLightSurfaceSample& Sample) const
 {
-	Sample.DiskPosition = FVector2D(0,0);
-	// Generate a sample on the surface of the sphere with uniform density over the surface area of the sphere
-	//@todo - stratify
-	const FVector4 UnitSpherePosition = GetUnitVector(RandomStream);
-	Sample.Position = UnitSpherePosition * LightSourceRadius + Position;
-	Sample.Normal = UnitSpherePosition;
-	// Probability of generating this surface position is 1 / SurfaceArea
-	Sample.PDF = 1.0f / (4.0f * (float)PI * LightSourceRadius * LightSourceRadius);
+	Sample.DiskPosition = FVector2D(0, 0);
+
+	if (LightSourceLength <= 0)
+	{
+		// Generate a sample on the surface of the sphere with uniform density over the surface area of the sphere
+		//@todo - stratify
+		const FVector4 UnitSpherePosition = GetUnitVector(RandomStream);
+		Sample.Position = UnitSpherePosition * LightSourceRadius + Position;
+		Sample.Normal = UnitSpherePosition;
+		// Probability of generating this surface position is 1 / SurfaceArea
+		Sample.PDF = 1.0f / (4.0f * (float)PI * LightSourceRadius * LightSourceRadius);
+	}
+	else
+	{
+		float CylinderSurfaceArea = 2.0f * (float)PI * LightSourceRadius * LightSourceLength;
+		float SphereSurfaceArea = 4.0f * (float)PI * LightSourceRadius * LightSourceRadius;
+		float TotalSurfaceArea = CylinderSurfaceArea + SphereSurfaceArea;
+
+		// Cylinder End caps
+		// The chance of calculating a point on the end sphere is equal to it's percentage of total surface area
+		if (RandomStream.GetFraction() < SphereSurfaceArea / TotalSurfaceArea)
+		{
+			// Generate a sample on the surface of the sphere with uniform density over the surface area of the sphere
+			//@todo - stratify
+			const FVector4 UnitSpherePosition = GetUnitVector(RandomStream);
+			Sample.Position = UnitSpherePosition * LightSourceRadius + Position;
+
+			if (Dot3(UnitSpherePosition, Direction) > 0)
+			{
+				Sample.Position += Direction * (LightSourceLength * 0.5f);
+			}
+			else
+			{
+				Sample.Position += -Direction * (LightSourceLength * 0.5f);
+			}
+
+			Sample.Normal = UnitSpherePosition;
+		}
+		// Cylinder body
+		else
+		{
+			// Get point along center line
+			FVector4 CentreLinePosition = Position + Direction * LightSourceLength * (RandomStream.GetFraction() - 0.5f);
+			// Get point radius away from center line at random angle
+			float Theta = 2.0f * (float)PI * RandomStream.GetFraction();
+			FVector4 CylEdgePos = FVector4(0, FMath::Cos(Theta), FMath::Sin(Theta), 1);
+			CylEdgePos = FRotationMatrix::MakeFromZ( Direction ).TransformVector( CylEdgePos );
+
+			Sample.Position = CylEdgePos * LightSourceRadius + CentreLinePosition;
+			Sample.Normal = CylEdgePos;
+		}
+
+		// Probability of generating this surface position is 1 / SurfaceArea
+		Sample.PDF = 1.0f / TotalSurfaceArea;
+	}
 }
 
 //----------------------------------------------------------------------------
@@ -1199,8 +1319,29 @@ FLinearColor FSpotLight::GetDirectIntensity(const FVector4& Point, bool bCalcula
 
 	if( LightFlags & GI_LIGHT_INVERSE_SQUARED )
 	{
-		float DistanceSqr = ( Position - Point ).SizeSquared3();
-		float DistanceAttenuation = 16.0f / ( DistanceSqr + 0.0001f );
+		FVector4 ToLight = Position - Point;
+		float DistanceSqr = ToLight.SizeSquared3();
+
+		float DistanceAttenuation = 0.0f;
+		if( LightSourceLength > 0.0f )
+		{
+			// Line segment irradiance
+			FVector4 L01 = Direction * LightSourceLength;
+			FVector4 L0 = ToLight - 0.5 * L01;
+			FVector4 L1 = ToLight + 0.5 * L01;
+			float LengthL0 = L0.Size3();
+			float LengthL1 = L1.Size3();
+
+			DistanceAttenuation = 1.0f / ( ( LengthL0 * LengthL1 + Dot3( L0, L1 ) ) * 0.5f + 1.0f );
+		}
+		else
+		{
+			// Sphere irradiance (technically just 1/d^2 but this avoids inf)
+			DistanceAttenuation = 1.0f / ( DistanceSqr + 1.0f );
+		}
+
+		// lumens
+		DistanceAttenuation *= 16.0f;
 
 		float LightRadiusMask = FMath::Square( FMath::Max( 0.0f, 1.0f - FMath::Square( DistanceSqr / (Radius * Radius) ) ) );
 		DistanceAttenuation *= LightRadiusMask;
diff --git a/Engine/Source/Programs/UnrealLightmass/Private/ImportExport/LightmassScene.h b/Engine/Source/Programs/UnrealLightmass/Private/ImportExport/LightmassScene.h
--- a/Engine/Source/Programs/UnrealLightmass/Private/ImportExport/LightmassScene.h
+++ b/Engine/Source/Programs/UnrealLightmass/Private/ImportExport/LightmassScene.h
@@ -178,7 +178,7 @@ public:
 	virtual void ValidateSurfaceSample(const FVector4& Point, FLightSurfaceSample& Sample) const {}
 
 	/** Gets a single position which represents the center of the area light source from the ReceivingPosition's point of view. */
-	virtual FVector4 LightCenterPosition(const FVector4& ReceivingPosition) const { return Position; }
+	virtual FVector4 LightCenterPosition(const FVector4& ReceivingPosition, const FVector4& ReceivingNormal) const { return Position; }
 
 	/** Returns true if all parts of the light are behind the surface being tested. */
 	virtual bool BehindSurface(const FVector4& TrianglePoint, const FVector4& TriangleNormal) const = 0;
@@ -265,7 +265,7 @@ public:
 	virtual void ValidateSurfaceSample(const FVector4& Point, FLightSurfaceSample& Sample) const;
 
 	/** Gets a single position which represents the center of the area light source from the ReceivingPosition's point of view. */
-	virtual FVector4 LightCenterPosition(const FVector4& ReceivingPosition) const;
+	virtual FVector4 LightCenterPosition(const FVector4& ReceivingPosition, const FVector4& ReceivingNormal) const;
 
 	/** Returns true if all parts of the light are behind the surface being tested. */
 	virtual bool BehindSurface(const FVector4& TrianglePoint, const FVector4& TriangleNormal) const;
@@ -373,6 +373,8 @@ public:
 	/** Returns the light's radiant power. */
 	virtual float Power() const;
 
+	virtual FVector4 LightCenterPosition(const FVector4& ReceivingPosition, const FVector4& ReceivingNormal) const;
+
 	/** Returns true if all parts of the light are behind the surface being tested. */
 	virtual bool BehindSurface(const FVector4& TrianglePoint, const FVector4& TriangleNormal) const;
 
diff --git a/Engine/Source/Programs/UnrealLightmass/Private/Lighting/LightingSystem.cpp b/Engine/Source/Programs/UnrealLightmass/Private/Lighting/LightingSystem.cpp
--- a/Engine/Source/Programs/UnrealLightmass/Private/Lighting/LightingSystem.cpp
+++ b/Engine/Source/Programs/UnrealLightmass/Private/Lighting/LightingSystem.cpp
@@ -2348,6 +2348,7 @@ bool FStaticLightingSystem::CalculatePointShadowing(
 		bool bIsShadowed = false;
 		if ((Light->LightFlags & GI_LIGHT_CASTSHADOWS) && (Light->LightFlags & GI_LIGHT_CASTSTATICSHADOWS))
 		{
+			// TODO find best point on light to shadow from
 			// Construct a line segment between the light and the surface point.
 			const FVector4 LightPosition = FVector4(Light->Position.X, Light->Position.Y, Light->Position.Z, 0);
 			const FVector4 LightVector = LightPosition - WorldSurfacePoint * Light->Position.W;
diff --git a/Engine/Source/Programs/UnrealLightmass/Private/Lighting/TextureMapping.cpp b/Engine/Source/Programs/UnrealLightmass/Private/Lighting/TextureMapping.cpp
--- a/Engine/Source/Programs/UnrealLightmass/Private/Lighting/TextureMapping.cpp
+++ b/Engine/Source/Programs/UnrealLightmass/Private/Lighting/TextureMapping.cpp
@@ -1890,7 +1890,7 @@ void FStaticLightingSystem::CalculateDirectSignedDistanceFieldLightingTextureMap
 					// And multiple shadowmaps on the same object may be merged together, but only if each one marks the area that it has valid data
 					CurrentSample.SetMapped(true);
 
-					const FVector4 LightPosition = Light->LightCenterPosition(TexelToVertex.WorldPosition);
+					const FVector4 LightPosition = Light->LightCenterPosition(TexelToVertex.WorldPosition, TexelToVertex.WorldTangentZ);
 					const FVector4 LightVector = (LightPosition - TexelToVertex.WorldPosition).SafeNormal();
 
 					FVector4 NormalForOffset = CurrentSample.GetNormal();
@@ -2121,7 +2121,7 @@ void FStaticLightingSystem::CalculateDirectSignedDistanceFieldLightingTextureMap
 							if ((bLightIsInFrontOfTriangle || bIsTwoSided) 
 								&& Light->AffectsBounds(FBoxSphereBounds(HighResSample.GetPosition(), FVector4(0,0,0),0)))
 							{
-								const FVector4 LightPosition = Light->LightCenterPosition(HighResSample.GetPosition());
+								const FVector4 LightPosition = Light->LightCenterPosition(HighResSample.GetPosition(), HighResSample.GetNormal());
 								const FVector4 LightVector = (LightPosition - HighResSample.GetPosition()).SafeNormal();
 
 								FVector4 NormalForOffset = HighResSample.GetNormal();
@@ -2312,6 +2312,7 @@ void FStaticLightingSystem::CalculateDirectSignedDistanceFieldLightingTextureMap
 											{
 												bool CurrentRegion = false;
 												FVector4 ScatterPosition;
+												FVector4 ScatterNormal;
 												bool bFoundScatterPosition = false;
 												
 												if (LowResScatterSample.NeedsHighResSampling())
@@ -2322,6 +2323,7 @@ void FStaticLightingSystem::CalculateDirectSignedDistanceFieldLightingTextureMap
 													{
 														CurrentRegion = HighResScatterSample.IsVisible();
 														ScatterPosition = HighResScatterSample.GetPosition();
+														ScatterNormal = HighResScatterSample.GetNormal();
 														bFoundScatterPosition = true;
 													}
 													else
@@ -2340,6 +2342,7 @@ void FStaticLightingSystem::CalculateDirectSignedDistanceFieldLightingTextureMap
 																	ClosestMappedSubSampleDistSquared = SubSampleDistanceSquared;
 																	CurrentRegion = SubHighResSample.IsVisible();
 																	ScatterPosition = SubHighResSample.GetPosition();
+																	ScatterNormal = SubHighResSample.GetNormal();
 																	bFoundScatterPosition = true;
 																}
 															}
@@ -2352,6 +2355,7 @@ void FStaticLightingSystem::CalculateDirectSignedDistanceFieldLightingTextureMap
 												{
 													CurrentRegion = LowResScatterSample.IsVisible();
 													ScatterPosition = LowResScatterSample.GetPosition();
+													ScatterNormal = LowResScatterSample.GetNormal();
 												}
 
 												// World space distance from the distance field texel to the nearest shadow transition
@@ -2376,7 +2380,7 @@ void FStaticLightingSystem::CalculateDirectSignedDistanceFieldLightingTextureMap
 													FinalShadowSample.Distance = CurrentRegion ? (NormalizedDistance) * .5f + .5f : .5f - NormalizedDistance * .5f;
 													// Approximate the penumbra size using PenumbraSize = (ReceiverDistanceFromLight - OccluderDistanceFromLight) * LightSize / OccluderDistanceFromLight,
 													// Which is from the paper "Percentage-Closer Soft Shadows" by Randima Fernando
-													const float ReceiverDistanceFromLight = (Light->LightCenterPosition(ScatterPosition) - ScatterPosition).Size3();
+													const float ReceiverDistanceFromLight = (Light->LightCenterPosition(ScatterPosition, ScatterNormal) - ScatterPosition).Size3();
 													// World space distance from center of penumbra to fully shadowed or fully lit transition
 													const float PenumbraSize = HighResSample.GetOccluderDistance() * Light->LightSourceRadius / (ReceiverDistanceFromLight - HighResSample.GetOccluderDistance());
 													// Normalize the penumbra size so it is a fraction of MaxTransitionDistanceWorldSpace
diff --git a/Engine/Source/Programs/UnrealLightmass/Public/SceneExport.h b/Engine/Source/Programs/UnrealLightmass/Public/SceneExport.h
--- a/Engine/Source/Programs/UnrealLightmass/Public/SceneExport.h
+++ b/Engine/Source/Programs/UnrealLightmass/Public/SceneExport.h
@@ -785,6 +785,8 @@ struct FLightData
 	float		Brightness;
 	/** The radius of the light's surface, not the light's influence. */
 	float		LightSourceRadius;
+	/** The length of the light source*/
+	float		LightSourceLength;
 	/** Scale factor for the indirect lighting */
 	float		IndirectLightingScale;
 	/** 0 will be completely desaturated, 1 will be unchanged, 2 will be completely saturated */
