Clamped static shadow depth map resolution on spotlights to be reasonable
* Allows building lighting in Infiltrator again which was generating a 23k x 23k shadowmap from a huge spotlight

diff --git a/Engine/Source/Programs/UnrealLightmass/Private/Lighting/LightingSystem.cpp b/Engine/Source/Programs/UnrealLightmass/Private/Lighting/LightingSystem.cpp
--- a/Engine/Source/Programs/UnrealLightmass/Private/Lighting/LightingSystem.cpp
+++ b/Engine/Source/Programs/UnrealLightmass/Private/Lighting/LightingSystem.cpp
@@ -2164,6 +2164,14 @@ void FStaticLightingSystem::CalculateStaticShadowDepthMap(FGuid LightGuid)
 		ShadowDepthMap->ShadowMapSizeX = ShadowDepthMap->ShadowMapSizeX == appTruncErrorCode ? INT_MAX : ShadowDepthMap->ShadowMapSizeX;
 		ShadowDepthMap->ShadowMapSizeY = ShadowDepthMap->ShadowMapSizeX;
 
+		// Clamp the number of dominant shadow samples generated if necessary while maintaining aspect ratio
+		if ((uint64)ShadowDepthMap->ShadowMapSizeX * (uint64)ShadowDepthMap->ShadowMapSizeY > (uint64)ShadowSettings.StaticShadowDepthMapMaxSamples)
+		{
+			const float AspectRatio = ShadowDepthMap->ShadowMapSizeX / (float)ShadowDepthMap->ShadowMapSizeY;
+			ShadowDepthMap->ShadowMapSizeY = FMath::TruncToInt(FMath::Sqrt(ShadowSettings.StaticShadowDepthMapMaxSamples / AspectRatio));
+			ShadowDepthMap->ShadowMapSizeX = FMath::TruncToInt(ShadowSettings.StaticShadowDepthMapMaxSamples / ShadowDepthMap->ShadowMapSizeY);
+		}
+
 		ShadowDepthMap->ShadowMap.Empty(ShadowDepthMap->ShadowMapSizeX * ShadowDepthMap->ShadowMapSizeY);
 		ShadowDepthMap->ShadowMap.AddZeroed(ShadowDepthMap->ShadowMapSizeX * ShadowDepthMap->ShadowMapSizeY);
 
