CIS build fix.

diff --git a/Engine/Source/Programs/UnrealLightmass/UnrealLightmass.Build.cs b/Engine/Source/Programs/UnrealLightmass/UnrealLightmass.Build.cs
--- a/Engine/Source/Programs/UnrealLightmass/UnrealLightmass.Build.cs
+++ b/Engine/Source/Programs/UnrealLightmass/UnrealLightmass.Build.cs
@@ -16,6 +16,11 @@ public class UnrealLightmass : ModuleRules
 
 			// Unreallightmass requires GetProcessMemoryInfo exported by psapi.dll. http://msdn.microsoft.com/en-us/library/windows/desktop/ms683219(v=vs.85).aspx
 			PublicAdditionalLibraries.Add("psapi.lib");
+			PrivateDependencyModuleNames.AddRange(
+				new string[] {
+					"Messaging",
+				}
+			);
 		}
 		else if (Target.Platform == UnrealTargetPlatform.Mac || Target.Platform == UnrealTargetPlatform.Linux)
 		{
