More fixes for Clang warnings

diff --git a/Engine/Plugins/Editor/SpeedTreeImporter/Source/SpeedTreeImporter/Private/SpeedTreeImportFactory.cpp b/Engine/Plugins/Editor/SpeedTreeImporter/Source/SpeedTreeImporter/Private/SpeedTreeImportFactory.cpp
--- a/Engine/Plugins/Editor/SpeedTreeImporter/Source/SpeedTreeImporter/Private/SpeedTreeImportFactory.cpp
+++ b/Engine/Plugins/Editor/SpeedTreeImporter/Source/SpeedTreeImporter/Private/SpeedTreeImportFactory.cpp
@@ -31,10 +31,19 @@
 
 #if WITH_SPEEDTREE
 
+#ifdef __clang__ // @todo: ThirdParty/SpeedTree/SpeedTreeSDK-v7.0/Include/Core/RenderState_inl.h:52:2: warning: field 'm_eHueVariation' will be initialized after field 'm_eShaderGenerationMode' [-Wreorder]
+	#pragma clang diagnostic push
+	#pragma clang diagnostic ignored "-Wreorder"
+#endif
+
 #include "Core/Core.h"
 
+#ifdef __clang__
+	#pragma clang diagnostic pop
 #endif
 
+#endif // WITH_SPEEDTREE
+
 #define LOCTEXT_NAMESPACE "SpeedTreeImportFactory"
 
 DEFINE_LOG_CATEGORY_STATIC(LogSpeedTreeImport, Log, All);
diff --git a/Engine/Plugins/OculusRift/Source/OculusRift/Private/OculusRiftRender.cpp b/Engine/Plugins/OculusRift/Source/OculusRift/Private/OculusRiftRender.cpp
--- a/Engine/Plugins/OculusRift/Source/OculusRift/Private/OculusRiftRender.cpp
+++ b/Engine/Plugins/OculusRift/Source/OculusRift/Private/OculusRiftRender.cpp
@@ -2,8 +2,18 @@
 //
 #include "OculusRiftPrivate.h"
 #include "OculusRiftHMD.h"
+
+#ifdef __clang__ // @todo: ThirdParty/Oculus/LibOVR/Include/../Src/OVR_DeviceMessages.h:227:46: warning: field 'CameraHandle' will be initialized after field 'pFrameData' [-Wreorder]
+	#pragma clang diagnostic push
+	#pragma clang diagnostic ignored "-Wreorder"
+#endif
+
 #include "../Src/OVR_Stereo.h"
 
+#ifdef __clang__
+	#pragma clang diagnostic pop
+#endif
+
 #include "RendererPrivate.h"
 #include "ScenePrivate.h"
 #include "PostProcess/PostProcessHMD.h"
diff --git a/Engine/Source/Developer/CrashDebugHelper/Public/CrashDebugHelper.h b/Engine/Source/Developer/CrashDebugHelper/Public/CrashDebugHelper.h
--- a/Engine/Source/Developer/CrashDebugHelper/Public/CrashDebugHelper.h
+++ b/Engine/Source/Developer/CrashDebugHelper/Public/CrashDebugHelper.h
@@ -207,8 +207,8 @@ struct FPDBCacheEntry
 {
 	/** Default constructor. */
 	FPDBCacheEntry( const FDateTime InLastAccessTime )
-		: SizeGB( 0 )
-		, LastAccessTime( InLastAccessTime )
+		: LastAccessTime( InLastAccessTime )
+		, SizeGB( 0 )
 	{}
 
 	/** Initialization constructor. */
diff --git a/Engine/Source/Editor/IntroTutorials/Private/STutorialsBrowser.cpp b/Engine/Source/Editor/IntroTutorials/Private/STutorialsBrowser.cpp
--- a/Engine/Source/Editor/IntroTutorials/Private/STutorialsBrowser.cpp
+++ b/Engine/Source/Editor/IntroTutorials/Private/STutorialsBrowser.cpp
@@ -39,6 +39,9 @@ public:
 		}
 	}
 
+	virtual ~FTutorialListEntry_Category()
+	{}
+
 	virtual TSharedRef<ITableRow> OnGenerateTutorialRow(const TSharedRef<STableViewBase>& OwnerTable) const override
 	{
 		return SNew(STableRow<TSharedPtr<ITutorialListEntry>>, OwnerTable)
@@ -168,6 +171,9 @@ public:
 		, HighlightText(InHighlightText)
 	{}
 
+	virtual ~FTutorialListEntry_Tutorial()
+	{}
+
 	virtual TSharedRef<ITableRow> OnGenerateTutorialRow(const TSharedRef<STableViewBase>& OwnerTable) const override
 	{
 		return SNew(STableRow<TSharedPtr<ITutorialListEntry>>, OwnerTable)
diff --git a/Engine/Source/Programs/UnrealLightmass/Private/Launch/UnrealLightmass.cpp b/Engine/Source/Programs/UnrealLightmass/Private/Launch/UnrealLightmass.cpp
--- a/Engine/Source/Programs/UnrealLightmass/Private/Launch/UnrealLightmass.cpp
+++ b/Engine/Source/Programs/UnrealLightmass/Private/Launch/UnrealLightmass.cpp
@@ -336,10 +336,10 @@ class FLocalCompareLightingResultsVisitor : public IPlatformFile::FDirectoryVisi
 public:
 
 	FLocalCompareLightingResultsVisitor(const TCHAR* InDir1, const TCHAR* InDir2, float InErrorThreshold)
-	:	Dir1(InDir1)
-	,	Dir2(InDir2)
-	,	NumDifferentFiles(0)
+	:	NumDifferentFiles(0)
 	,	TotalFiles(0)
+	,	Dir1(InDir1)
+	,	Dir2(InDir2)
 	,	ErrorThreshold(InErrorThreshold)
 	{
 	}
diff --git a/Engine/Source/Programs/UnrealLightmass/Private/Lighting/LightingSystem.h b/Engine/Source/Programs/UnrealLightmass/Private/Lighting/LightingSystem.h
--- a/Engine/Source/Programs/UnrealLightmass/Private/Lighting/LightingSystem.h
+++ b/Engine/Source/Programs/UnrealLightmass/Private/Lighting/LightingSystem.h
@@ -121,9 +121,9 @@ public:
 	{}
 
 	FGatheredLightMapSample(const FGatheredLightSample& Sample)
-	: bIsMapped(false)
-	, HighQuality(Sample)
+	: HighQuality(Sample)
 	, LowQuality(Sample)
+	, bIsMapped(false)
 	{}
 
 	FGatheredLightMapSample& operator=(const FGatheredLightSample& Sample)
@@ -1455,18 +1455,18 @@ public:
 
 	/** Initialization constructor. */
 	FStaticLightingThreadRunnable(FStaticLightingSystem* InSystem, int32 InThreadIndex) :
-		System(InSystem),
 		Thread(NULL),
 		IdleTime(0),
 		ThreadIndex(InThreadIndex),
+		System(InSystem),
 		bTerminatedByError(false)
 	{}
 
 	FStaticLightingThreadRunnable(FStaticLightingSystem* InSystem) :
-		System(InSystem),
 		Thread(NULL),
 		IdleTime(0),
 		ThreadIndex(0),
+		System(InSystem),
 		bTerminatedByError(false)
 	{}
 
@@ -1558,8 +1558,8 @@ public:
 		int32 InThreadIndex,
 		const FDirectPhotonEmittingInput& InInput)
 		:
-		ThreadIndex(InThreadIndex),
 		FStaticLightingThreadRunnable(InSystem),
+		ThreadIndex(InThreadIndex),
 		Input(InInput)
 	{}
 
diff --git a/Engine/Source/Programs/UnrealLightmass/Private/Lighting/TextureMapping.cpp b/Engine/Source/Programs/UnrealLightmass/Private/Lighting/TextureMapping.cpp
--- a/Engine/Source/Programs/UnrealLightmass/Private/Lighting/TextureMapping.cpp
+++ b/Engine/Source/Programs/UnrealLightmass/Private/Lighting/TextureMapping.cpp
@@ -2986,9 +2986,9 @@ public:
 		) :
 		Scene(InScene),
 		TexelToNumTrianglesMap(InTexelToNumTrianglesMap),
-		bDebugThisMapping(bInDebugThisMapping),
 		TotalPixelsWritten(0),
-		TotalPixelOverlapsOccured(0)
+		TotalPixelOverlapsOccured(0),
+		bDebugThisMapping(bInDebugThisMapping)
 	{}
 
 	int32 GetTotalPixelsWritten() const
diff --git a/Engine/Source/Programs/UnrealLightmass/Private/LightmassCore/Math/LMkDOP.h b/Engine/Source/Programs/UnrealLightmass/Private/LightmassCore/Math/LMkDOP.h
--- a/Engine/Source/Programs/UnrealLightmass/Private/LightmassCore/Math/LMkDOP.h
+++ b/Engine/Source/Programs/UnrealLightmass/Private/LightmassCore/Math/LMkDOP.h
@@ -422,9 +422,9 @@ public:
 		V0(vert0), V1(vert1), V2(vert2),
 		MeshIndex(InMeshIndex),
 		LODIndex(InLODIndex),
-		MaterialIndex(InMaterialIndex),
 		bTwoSided(bInTwoSided),
-		bStaticAndOpaque(bInStaticAndOpaque)
+		bStaticAndOpaque(bInStaticAndOpaque),
+		MaterialIndex(InMaterialIndex)
 	{
 #if CACHE_BUILD_TEMPORARIES
 		// Calculate the centroid for the triangle
diff --git a/Engine/Source/Runtime/Media/Public/Common/MediaSampleBuffer.h b/Engine/Source/Runtime/Media/Public/Common/MediaSampleBuffer.h
--- a/Engine/Source/Runtime/Media/Public/Common/MediaSampleBuffer.h
+++ b/Engine/Source/Runtime/Media/Public/Common/MediaSampleBuffer.h
@@ -24,6 +24,8 @@ public:
 		: CurrentSampleTime(FTimespan::MinValue())
 	{ }
 
+	virtual ~FMediaSampleBuffer()
+	{ }
 public:
 
 	/**
diff --git a/Engine/Source/Runtime/Media/Public/Common/MediaSampleQueue.h b/Engine/Source/Runtime/Media/Public/Common/MediaSampleQueue.h
--- a/Engine/Source/Runtime/Media/Public/Common/MediaSampleQueue.h
+++ b/Engine/Source/Runtime/Media/Public/Common/MediaSampleQueue.h
@@ -19,6 +19,8 @@ class FMediaSampleQueue
 {
 public:
 
+	virtual ~FMediaSampleQueue() {}
+
 	// IMediaSink interface
 
 	void ProcessMediaSample( const void* Buffer, uint32 BufferSize, FTimespan Duration, FTimespan Time ) override
diff --git a/Engine/Source/Runtime/MediaAssets/Private/Assets/MediaTextureResource.cpp b/Engine/Source/Runtime/MediaAssets/Private/Assets/MediaTextureResource.cpp
--- a/Engine/Source/Runtime/MediaAssets/Private/Assets/MediaTextureResource.cpp
+++ b/Engine/Source/Runtime/MediaAssets/Private/Assets/MediaTextureResource.cpp
@@ -7,9 +7,9 @@
  *****************************************************************************/
 
 FMediaTextureResource::FMediaTextureResource( const class UMediaTexture* InOwner, const FMediaSampleBufferRef& InVideoBuffer )
-	: Owner(InOwner)
-	, Cleared(false)
+	: Cleared(false)
 	, LastFrameTime(FTimespan::MinValue())
+	, Owner(InOwner)
 	, VideoBuffer(InVideoBuffer)
 { }
 
diff --git a/Engine/Source/Runtime/Online/BuildPatchServices/Private/BuildPatchGeneration.cpp b/Engine/Source/Runtime/Online/BuildPatchServices/Private/BuildPatchGeneration.cpp
--- a/Engine/Source/Runtime/Online/BuildPatchServices/Private/BuildPatchGeneration.cpp
+++ b/Engine/Source/Runtime/Online/BuildPatchServices/Private/BuildPatchGeneration.cpp
@@ -288,16 +288,16 @@ FBuildDataChunkProcessor::FBuildDataChunkProcessor( FBuildPatchAppManifestRef In
 	: NumNewChunks( 0 )
 	, NumKnownChunks( 0 )
 	, BuildRoot( InBuildRoot )
-	, ChunkIsPushed( false )
+	, ChunkWriter( FBuildPatchServicesModule::GetCloudDirectory() )
+	, CurrentChunkBufferPos( 0 )
+	, CurrentFile( NULL )
 	, IsProcessingChunk( false )
 	, IsProcessingChunkPart( false )
+	, IsProcessingFile( false )
+	, ChunkIsPushed( false )
+	, BackupChunkBufferPos( 0 )
 	, BackupProcessingChunk( false )
 	, BackupProcessingChunkPart( false )
-	, CurrentChunkBufferPos( 0 )
-	, BackupChunkBufferPos( 0 )
-	, IsProcessingFile( false )
-	, CurrentFile( NULL )
-	, ChunkWriter( FBuildPatchServicesModule::GetCloudDirectory() )
 {
 	BuildManifest = InBuildManifest;
 	CurrentChunkGuid.Invalidate();
