Safe changes to Lightmass for Linux.

Part of PR #316.

diff --git a/Engine/Source/Editor/UnrealEd/Private/StaticLightingSystem/StaticLightingPrivate.h b/Engine/Source/Editor/UnrealEd/Private/StaticLightingSystem/StaticLightingPrivate.h
--- a/Engine/Source/Editor/UnrealEd/Private/StaticLightingSystem/StaticLightingPrivate.h
+++ b/Engine/Source/Editor/UnrealEd/Private/StaticLightingSystem/StaticLightingPrivate.h
@@ -403,7 +403,7 @@ private:
  * NOTE: These must remain binary compatible with the ones in Lightmass.
  */
 
-#if !PLATFORM_MAC
+#if !PLATFORM_MAC && !PLATFORM_LINUX
 	#pragma pack(push, 1)
 #endif
 
@@ -493,7 +493,7 @@ struct FDebugLightingOutput
 	{}
 };
 
-#if !PLATFORM_MAC
+#if !PLATFORM_MAC && !PLATFORM_LINUX
 	#pragma pack(pop)
 #endif
 
diff --git a/Engine/Source/Programs/UnrealLightmass/Public/ImportExport.h b/Engine/Source/Programs/UnrealLightmass/Public/ImportExport.h
--- a/Engine/Source/Programs/UnrealLightmass/Public/ImportExport.h
+++ b/Engine/Source/Programs/UnrealLightmass/Public/ImportExport.h
@@ -67,7 +67,7 @@ namespace Lightmass
 		SOURCEOBJECTTYPE_Mapping
 	};
 
-#if !PLATFORM_MAC
+#if !PLATFORM_MAC && !PLATFORM_LINUX
 	#pragma pack(push, 1)
 #endif
 
@@ -435,7 +435,7 @@ namespace Lightmass
 	/** Guid used by Unreal to determine when the volume lighting sample channel with the same Guid can be opened. */
 	static const FGuid PrecomputedVolumeLightingGuid = FGuid(0xce97c5c3, 0xab614fd3, 0xb2da55c0, 0xe6c33fb4);
 
-#if !PLATFORM_MAC
+#if !PLATFORM_MAC && !PLATFORM_LINUX
 	#pragma pack(pop)
 #endif
 
diff --git a/Engine/Source/Programs/UnrealLightmass/Public/MaterialExport.h b/Engine/Source/Programs/UnrealLightmass/Public/MaterialExport.h
--- a/Engine/Source/Programs/UnrealLightmass/Public/MaterialExport.h
+++ b/Engine/Source/Programs/UnrealLightmass/Public/MaterialExport.h
@@ -5,7 +5,7 @@
 namespace Lightmass
 {
 
-#if !PLATFORM_MAC
+#if !PLATFORM_MAC && !PLATFORM_LINUX
 #pragma pack(push, 1)
 #endif
 
@@ -82,7 +82,7 @@ struct FMaterialData
 	}
 };
 
-#if !PLATFORM_MAC
+#if !PLATFORM_MAC && !PLATFORM_LINUX
 #pragma pack(pop)
 #endif
 
diff --git a/Engine/Source/Programs/UnrealLightmass/Public/MeshExport.h b/Engine/Source/Programs/UnrealLightmass/Public/MeshExport.h
--- a/Engine/Source/Programs/UnrealLightmass/Public/MeshExport.h
+++ b/Engine/Source/Programs/UnrealLightmass/Public/MeshExport.h
@@ -7,7 +7,7 @@ namespace Lightmass
 
 enum { MAX_TEXCOORDS=4 };
 
-#if !PLATFORM_MAC
+#if !PLATFORM_MAC && !PLATFORM_LINUX
 #pragma pack(push, 1)
 #endif
 
@@ -96,7 +96,7 @@ struct FStaticMeshVertex
 	FVector2D		UVs[MAX_TEXCOORDS];
 };
 
-#if !PLATFORM_MAC
+#if !PLATFORM_MAC && !PLATFORM_LINUX
 #pragma pack(pop)
 #endif
 
diff --git a/Engine/Source/Programs/UnrealLightmass/Public/SceneExport.h b/Engine/Source/Programs/UnrealLightmass/Public/SceneExport.h
--- a/Engine/Source/Programs/UnrealLightmass/Public/SceneExport.h
+++ b/Engine/Source/Programs/UnrealLightmass/Public/SceneExport.h
@@ -7,7 +7,7 @@
 namespace Lightmass
 {
 
-#if !PLATFORM_MAC
+#if !PLATFORM_MAC && !PLATFORM_LINUX
 #pragma pack(push, 1)
 #endif
 
@@ -1096,7 +1096,7 @@ struct FLandscapeStaticLightingMeshData
 	int32 ExpandQuadsY;
 };
 
-#if !PLATFORM_MAC
+#if !PLATFORM_MAC && !PLATFORM_LINUX
 #pragma pack(pop)
 #endif
 
