Removed MaxTrianglesPerLeaf which was legacy
Restored Lightmass to parsing parameters passed on the commandline.  Mostly this fixes NumUnusedLightmassThreads to work again.
Changed -debug to -usedebug as it is a dangerous command, some projects have -debug on the commandline by default
Contains changes from PR #316

diff --git a/Engine/Config/BaseLightmass.ini b/Engine/Config/BaseLightmass.ini
--- a/Engine/Config/BaseLightmass.ini
+++ b/Engine/Config/BaseLightmass.ini
@@ -18,7 +18,7 @@ DefaultStaticMeshLightingRes=32
 bAllowCropping=False
 bGarbageCollectAfterExport=True
 bRebuildDirtyGeometryForLighting=True
-NumUnusedLightmassThreads=3
+NumUnusedLightmassThreads=2
 
 [DevOptions.StaticLightingSceneConstants]
 StaticLightingLevelScale=1
diff --git a/Engine/Source/Editor/UnrealEd/Private/Lightmass/Lightmass.cpp b/Engine/Source/Editor/UnrealEd/Private/Lightmass/Lightmass.cpp
--- a/Engine/Source/Editor/UnrealEd/Private/Lightmass/Lightmass.cpp
+++ b/Engine/Source/Editor/UnrealEd/Private/Lightmass/Lightmass.cpp
@@ -2605,14 +2605,6 @@ bool FLightmassProcessor::BeginRun()
 
 	CommandLineParameters += FString::Printf(TEXT(" -numthreads %i"), NumThreads);
 
-	//append on the maximum number of triangles per leaf for the kdop tree
-	uint32 MaxTrisPerLeaf = 4;
-	if (System.GetWorld()->GetWorldSettings())
-	{
-		MaxTrisPerLeaf = System.GetWorld()->GetWorldSettings()->MaxTrianglesPerLeaf;
-	}
-	CommandLineParameters += FString::Printf(TEXT(" -trisperleaf %i"), MaxTrisPerLeaf);
-
 	NSwarm::FJobSpecification JobSpecification32, JobSpecification64;
 	if ( !bUse64bitProcess )
 	{
diff --git a/Engine/Source/Programs/UnrealLightmass/Private/CPUSolver/CPUSolver.h b/Engine/Source/Programs/UnrealLightmass/Private/CPUSolver/CPUSolver.h
--- a/Engine/Source/Programs/UnrealLightmass/Private/CPUSolver/CPUSolver.h
+++ b/Engine/Source/Programs/UnrealLightmass/Private/CPUSolver/CPUSolver.h
@@ -101,7 +101,10 @@ extern class FLightmassSwarm* GSwarm;
 /** Whether we should report detailed stats back to Unreal. */
 extern bool GReportDetailedStats;
 
-/** Whether Lightmass is running in debug mode (-debug), using a hardcoded job and not requesting tasks from Swarm. */
+/** 
+ * Whether Lightmass is running in debug mode (-debug), using a hardcoded job and not requesting tasks from Swarm. 
+ * Warning!  This will only process mapping tasks and will skip other types of tasks.
+ */
 extern bool GDebugMode;
 
 } //namespace Lightmass
diff --git a/Engine/Source/Programs/UnrealLightmass/Private/Launch/UnrealLightmass.cpp b/Engine/Source/Programs/UnrealLightmass/Private/Launch/UnrealLightmass.cpp
--- a/Engine/Source/Programs/UnrealLightmass/Private/Launch/UnrealLightmass.cpp
+++ b/Engine/Source/Programs/UnrealLightmass/Private/Launch/UnrealLightmass.cpp
@@ -117,7 +117,7 @@ int LightmassMain(int argc, ANSICHAR* argv[])
 
 	for (int32 ArgIndex = 1; ArgIndex < argc; ArgIndex++)
 	{
-		if ((FCStringAnsi::Stricmp(argv[ArgIndex], " -help") == 0) || (FCStringAnsi::Stricmp(argv[ArgIndex], " -?") == 0))
+		if ((FCStringAnsi::Stricmp(argv[ArgIndex], "-help") == 0) || (FCStringAnsi::Stricmp(argv[ArgIndex], "-?") == 0))
 		{
 			UE_LOG(LogLightmass, Display, TEXT("Usage:\n  UnrealLightmass\n\t[SceneGuid]\n\t[-debug]\n\t[-unittest]\n\t[-dumptex]\n\t[-numthreads N]\n\t[-compare Dir1 Dir2 [-error N]]"));
 			UE_LOG(LogLightmass, Display, TEXT(""));
@@ -129,23 +129,24 @@ int LightmassMain(int argc, ANSICHAR* argv[])
 			UE_LOG(LogLightmass, Display, TEXT("  -error : Controls the threshold that an error is counted when comparing with -compare"));
 			return 0;
 		}
-		else if (FCStringAnsi::Stricmp(argv[ArgIndex], " -unittest") == 0)
+		else if (FCStringAnsi::Stricmp(argv[ArgIndex], "-unittest") == 0)
 		{
 			bRunUnitTest = true;
 		}
-		else if (FCStringAnsi::Stricmp(argv[ArgIndex], " -dumptex") == 0)
+		else if (FCStringAnsi::Stricmp(argv[ArgIndex], "-dumptex") == 0)
 		{
 			bDumpTextures = true;
 		}
-		else if (FCStringAnsi::Stricmp(argv[ArgIndex], " -debug") == 0)
+		else if (FCStringAnsi::Stricmp(argv[ArgIndex], "-usedebug") == 0)
 		{
+			// Warning!  This will only process mapping tasks and will skip other types of tasks.
 			GDebugMode = true;
 		}
-		else if (FCStringAnsi::Stricmp(argv[ArgIndex], " -stats") == 0)
+		else if (FCStringAnsi::Stricmp(argv[ArgIndex], "-stats") == 0)
 		{
 			GReportDetailedStats = true;
 		}
-		else if (FCStringAnsi::Stricmp(argv[ArgIndex], " -numthreads") == 0)
+		else if (FCStringAnsi::Stricmp(argv[ArgIndex], "-numthreads") == 0)
 		{
 			// use the next parameter as the number of threads (it must exist, or we fail)
 			NumThreads = 0;
@@ -161,7 +162,7 @@ int LightmassMain(int argc, ANSICHAR* argv[])
 				return 1;
 			}
 		}
-		else if (FCStringAnsi::Stricmp(argv[ArgIndex], " -compare") == 0)
+		else if (FCStringAnsi::Stricmp(argv[ArgIndex], "-compare") == 0)
 		{
 			bCompareFiles = true;
 
@@ -174,7 +175,7 @@ int LightmassMain(int argc, ANSICHAR* argv[])
 			File1 = *FString(argv[++ArgIndex]);
 			File2 = *FString(argv[++ArgIndex]);
 		}
-		else if (FCStringAnsi::Stricmp(argv[ArgIndex], " -error") == 0)
+		else if (FCStringAnsi::Stricmp(argv[ArgIndex], "-error") == 0)
 		{
 			// use the next parameter as the number of threads (it must exist, or we fail)
 			if (ArgIndex >= argc - 1)
@@ -204,17 +205,6 @@ int LightmassMain(int argc, ANSICHAR* argv[])
 			SceneGuid.D = wcstoul(*Arg.Mid(24, 8), NULL, 16);
 #endif
 		}
-		else if (FCStringAnsi::Stricmp(argv[ArgIndex], " -trisperleaf") == 0)
-		{
-			// use the next parameter as the maximum number of triangles per leaf for the kdop tree (it must exist, or we fail)
-			if (ArgIndex >= argc - 1)
-			{
-				UE_LOG(LogLightmass, Display, TEXT("-error requires an error value following (-trisperleaf N)"));
-				return 1;
-			}
-
-			GKDOPMaxTrisPerLeaf = FCString::Atoi(*FString(argv[++ArgIndex]));
-		}
 	}
 
 	// if we want to run the unit test, do that, then nothing else
diff --git a/Engine/Source/Runtime/Engine/Classes/GameFramework/WorldSettings.h b/Engine/Source/Runtime/Engine/Classes/GameFramework/WorldSettings.h
--- a/Engine/Source/Runtime/Engine/Classes/GameFramework/WorldSettings.h
+++ b/Engine/Source/Runtime/Engine/Classes/GameFramework/WorldSettings.h
@@ -307,12 +307,6 @@ class ENGINE_API AWorldSettings : public AInfo, public IInterface_AssetUserData
 	UPROPERTY(EditAnywhere, Category=Lightmass)
 	uint32 bForceNoPrecomputedLighting:1;
 
-#if WITH_EDITORONLY_DATA
-	/** The number of triangles per leaf of the kdop tree in lightmass*/
-	UPROPERTY(meta=(FixedIncrement = "4", ClampMin = "4", Multiple = "4"))
-	int32 MaxTrianglesPerLeaf;
-
-#endif // WITH_EDITORONLY_DATA
 	UPROPERTY(EditAnywhere, Category=Lightmass)
 	struct FLightmassWorldInfoSettings LightmassSettings;
 
diff --git a/Engine/Source/Runtime/Engine/Private/WorldSettings.cpp b/Engine/Source/Runtime/Engine/Private/WorldSettings.cpp
--- a/Engine/Source/Runtime/Engine/Private/WorldSettings.cpp
+++ b/Engine/Source/Runtime/Engine/Private/WorldSettings.cpp
@@ -46,7 +46,6 @@ AWorldSettings::AWorldSettings(const class FPostConstructInitializeProperties& P
 	PackedLightAndShadowMapTextureSize = 1024;
 #if WITH_EDITORONLY_DATA
 	bHiddenEd = true;
-	MaxTrianglesPerLeaf = 4;
 #endif // WITH_EDITORONLY_DATA
 
 	DefaultColorScale = FVector(1.0f, 1.0f, 1.0f);
