Fixed IndirectLightingSaturation not being applied to first bounce GI

diff --git a/Engine/Source/Programs/UnrealLightmass/Private/ImportExport/LightmassScene.cpp b/Engine/Source/Programs/UnrealLightmass/Private/ImportExport/LightmassScene.cpp
--- a/Engine/Source/Programs/UnrealLightmass/Private/ImportExport/LightmassScene.cpp
+++ b/Engine/Source/Programs/UnrealLightmass/Private/ImportExport/LightmassScene.cpp
@@ -471,7 +471,14 @@ FLinearColor FLight::GetDirectIntensity(const FVector4& Point, bool bCalculateFo
 		LightProfileAttenuation = ComputeLightProfileMultiplier(Dot3(NegLightVector, Direction));
 	}
 
-	return FLinearColor(Color) * (LightProfileAttenuation * Brightness) * (bCalculateForIndirectLighting ? IndirectLightingScale : 1.0f);
+	if (bCalculateForIndirectLighting)
+	{
+		return IndirectColor * (LightProfileAttenuation * Brightness);
+	}
+	else
+	{
+		return FLinearColor(Color) * (LightProfileAttenuation * Brightness);
+	}
 }
 
 /** Generates and caches samples on the light's surface. */
