Static meshes completely in shadow from a stationary light will toss their static shadowmap data and use a cheaper base pass permutation

diff --git a/Engine/Source/Programs/UnrealLightmass/Private/Lighting/TextureMapping.cpp b/Engine/Source/Programs/UnrealLightmass/Private/Lighting/TextureMapping.cpp
--- a/Engine/Source/Programs/UnrealLightmass/Private/Lighting/TextureMapping.cpp
+++ b/Engine/Source/Programs/UnrealLightmass/Private/Lighting/TextureMapping.cpp
@@ -542,28 +542,6 @@ void FStaticLightingSystem::ProcessTextureMapping(FStaticLightingTextureMapping*
 		// Release corner information as it is no longer needed
 		TexelToCornersMap.Empty();
 
-		int32 NumAffectingDominantLights = 0;
-		for (TMap<const FLight*,FSignedDistanceFieldShadowMapData2D*>::TIterator It(SignedDistanceFieldShadowMaps); It; ++It)
-		{
-			if (It.Key()->LightFlags & GI_LIGHT_DOMINANT)
-			{
-				NumAffectingDominantLights++;
-			}
-		}
-
-		for (TMap<const FLight*,FShadowMapData2D*>::TIterator It(ShadowMaps); It; ++It)
-		{
-			if (It.Key()->LightFlags & GI_LIGHT_DOMINANT)
-			{
-				NumAffectingDominantLights++;
-			}
-		}
-
-		if (NumAffectingDominantLights > 1) 
-		{
-			GSwarm->SendAlertMessage(NSwarm::ALERT_LEVEL_ERROR, TextureMapping->Guid, SOURCEOBJECTTYPE_Mapping, TEXT("LightmassError_ObjectMultipleDominantLights"));
-		}
-
 		if (bDebugThisMapping)
 		{
 			int32 asdf = 0;
@@ -1199,7 +1177,7 @@ void FStaticLightingSystem::CalculateDirectLightingTextureMappingFiltered(
 		}
 	}
 
-	if(bIsCompletelyOccluded && !TextureMapping->Mesh->bInstancedStaticMesh)
+	if(bIsCompletelyOccluded) 
 	{
 		// If the light is completely occluded, discard the shadow-map.
 		delete FilteredShadowMapData;
@@ -1643,7 +1621,7 @@ void FStaticLightingSystem::CalculateDirectAreaLightingTextureMapping(
 	}
 	else if (ShadowMapData)
 	{
-		if ((bIsCompletelyOccluded || NumUnoccludedTexels < NumMappedTexels * ShadowSettings.MinUnoccludedFraction) && !TextureMapping->Mesh->bInstancedStaticMesh)
+		if (bIsCompletelyOccluded || NumUnoccludedTexels < NumMappedTexels * ShadowSettings.MinUnoccludedFraction)
 		{
 			delete ShadowMapData;
 		}
@@ -1941,7 +1919,7 @@ void FStaticLightingSystem::CalculateDirectSignedDistanceFieldLightingTextureMap
 	}
 	FirstPassSourceTimer.Stop();
 
-	if ( (!bIsCompletelyOccluded && NumUnoccludedTexels > NumMappedTexels * ShadowSettings.MinUnoccludedFraction) || TextureMapping->Mesh->bInstancedStaticMesh)
+	if (!bIsCompletelyOccluded && NumUnoccludedTexels > NumMappedTexels * ShadowSettings.MinUnoccludedFraction)
 	{
 		LIGHTINGSTAT(FManualRDTSCTimer SecondPassSourceTimer(MappingContext.Stats.SignedDistanceFieldSourceSecondPassThreadTime));
 		check(UpsampleFactor % 2 == 1 && UpsampleFactor >= 1);
diff --git a/Engine/Source/Programs/UnrealLightmass/Public/SceneExport.h b/Engine/Source/Programs/UnrealLightmass/Public/SceneExport.h
--- a/Engine/Source/Programs/UnrealLightmass/Public/SceneExport.h
+++ b/Engine/Source/Programs/UnrealLightmass/Public/SceneExport.h
@@ -768,7 +768,6 @@ enum EDawnLightFlags
 	// maps to ULightComponent::CastStaticShadows
 	GI_LIGHT_CASTSTATICSHADOWS		= 0x00000010,
 	GI_LIGHT_USESIGNEDDISTANCEFIELDSHADOWS = 0x00000020,
-	GI_LIGHT_DOMINANT				= 0x00000040,
 	GI_LIGHT_INVERSE_SQUARED		= 0x00000080,
 	GI_LIGHT_USE_LIGHTPROFILE		= 0x00000100
 };
@@ -948,8 +947,6 @@ struct FStaticLightingMeshInstanceData
 	bool bCastShadowAsTwoSided;
 	/** Whether the mesh can be moved in game or not. */
 	bool bMovable;
-	/** Whether this is an instanced static mesh */
-	bool bInstancedStaticMesh;
 	/** The lights which affect the mesh's primitive. */
 	int32 NumRelevantLights;
 	/** The bounding box of the mesh. */
