LoadVector3* macros replaced with their componentwise equivalents from UnrealMathFPU.h to prevent GuardMalloc reports and/or page faults.

diff --git a/Engine/Source/Programs/UnrealLightmass/Private/LightmassCore/Math/LMMathSSE.h b/Engine/Source/Programs/UnrealLightmass/Private/LightmassCore/Math/LMMathSSE.h
--- a/Engine/Source/Programs/UnrealLightmass/Private/LightmassCore/Math/LMMathSSE.h
+++ b/Engine/Source/Programs/UnrealLightmass/Private/LightmassCore/Math/LMMathSSE.h
@@ -127,7 +127,7 @@ static const VectorRegister SSE_SIGN_MASK = MakeVectorRegister( (uint32)SIGN_BIT
  * @param Ptr	Unaligned memory pointer to the 3 FLOATs
  * @return		VectorRegister(Ptr[0], Ptr[1], Ptr[2], undefined)
  */
-#define VectorLoadFloat3( Ptr )			_mm_loadu_ps( (float*)(Ptr) )
+#define VectorLoadFloat3( Ptr )			MakeVectorRegister( ((const float*)(Ptr))[0], ((const float*)(Ptr))[1], ((const float*)(Ptr))[2], 0.0f )
 
 /**
  * Loads 3 FLOATs from unaligned memory and sets W=0.
@@ -135,7 +135,7 @@ static const VectorRegister SSE_SIGN_MASK = MakeVectorRegister( (uint32)SIGN_BIT
  * @param Ptr	Unaligned memory pointer to the 3 FLOATs
  * @return		VectorRegister(Ptr[0], Ptr[1], Ptr[2], 0.0f)
  */
-#define VectorLoadFloat3_W0( Ptr )		_mm_and_ps( _mm_loadu_ps((float*)(Ptr)), SSE_XYZ_MASK )
+#define VectorLoadFloat3_W0( Ptr )		MakeVectorRegister( ((const float*)(Ptr))[0], ((const float*)(Ptr))[1], ((const float*)(Ptr))[2], 0.0f )
 
 /**
  * Loads 3 FLOATs from unaligned memory and sets W=1.
@@ -143,17 +143,7 @@ static const VectorRegister SSE_SIGN_MASK = MakeVectorRegister( (uint32)SIGN_BIT
  * @param Ptr	Unaligned memory pointer to the 3 FLOATs
  * @return		VectorRegister(Ptr[0], Ptr[1], Ptr[2], 1.0f)
  */
-FORCEINLINE VectorRegister VectorLoadFloat3_W1( const void* Ptr )
-{
-	// Vec = (Ptr[0], Ptr[1], Ptr[2], undefined)
-	VectorRegister Vec = _mm_loadu_ps((const float*)Ptr);
-
-	// Temp = (Ptr[2], undefined, 1.0f, 1.0f)
-	VectorRegister Temp = _mm_movehl_ps( VectorOne(), Vec );
-
-	// Return (Ptr[0], Ptr[1], Ptr[2], 1.0f)
-	return _mm_shuffle_ps( Vec, Temp, SHUFFLEMASK(0,1,0,3) );
-}
+#define VectorLoadFloat3_W1( Ptr )		MakeVectorRegister( ((const float*)(Ptr))[0], ((const float*)(Ptr))[1], ((const float*)(Ptr))[2], 1.0f )
 
 /**
  * Loads 4 FLOATs from aligned memory.
diff --git a/Engine/Source/Runtime/Core/Public/Math/UnrealMathDirectX.h b/Engine/Source/Runtime/Core/Public/Math/UnrealMathDirectX.h
--- a/Engine/Source/Runtime/Core/Public/Math/UnrealMathDirectX.h
+++ b/Engine/Source/Runtime/Core/Public/Math/UnrealMathDirectX.h
@@ -96,7 +96,7 @@ FORCEINLINE VectorRegister MakeVectorRegister( float X, float Y, float Z, float
  * @param Ptr	Unaligned memory pointer to the 3 FLOATs
  * @return		VectorRegister(Ptr[0], Ptr[1], Ptr[2], undefined)
  */
-#define VectorLoadFloat3( Ptr )		DirectX::XMLoadFloat4( (const DirectX::XMFLOAT4*)(Ptr) )	
+#define VectorLoadFloat3( Ptr )			MakeVectorRegister( ((const float*)(Ptr))[0], ((const float*)(Ptr))[1], ((const float*)(Ptr))[2], 0.0f )
 
 
 /**
@@ -105,7 +105,7 @@ FORCEINLINE VectorRegister MakeVectorRegister( float X, float Y, float Z, float
  * @param Ptr	Unaligned memory pointer to the 3 FLOATs
  * @return		VectorRegister(Ptr[0], Ptr[1], Ptr[2], 0.0f)
  */
-#define VectorLoadFloat3_W0( Ptr )		DirectX::XMVectorAndInt( VectorLoadFloat3( Ptr ), DirectX::g_XMMask3 )
+#define VectorLoadFloat3_W0( Ptr )		MakeVectorRegister( ((const float*)(Ptr))[0], ((const float*)(Ptr))[1], ((const float*)(Ptr))[2], 0.0f )
 
 
 /**
@@ -114,8 +114,8 @@ FORCEINLINE VectorRegister MakeVectorRegister( float X, float Y, float Z, float
  * @param Ptr	Unaligned memory pointer to the 3 FLOATs
  * @return		VectorRegister(Ptr[0], Ptr[1], Ptr[2], 1.0f)
  */
-#define VectorLoadFloat3_W1( Ptr )		DirectX::XMVectorPermute<0,1,2,7>( VectorLoadFloat3( Ptr ), VectorOne() )
-	
+#define VectorLoadFloat3_W1( Ptr )		MakeVectorRegister( ((const float*)(Ptr))[0], ((const float*)(Ptr))[1], ((const float*)(Ptr))[2], 1.0f )
+
 /**
  * Loads 4 FLOATs from aligned memory.
  *
diff --git a/Engine/Source/Runtime/Core/Public/Math/UnrealMathNeon.h b/Engine/Source/Runtime/Core/Public/Math/UnrealMathNeon.h
--- a/Engine/Source/Runtime/Core/Public/Math/UnrealMathNeon.h
+++ b/Engine/Source/Runtime/Core/Public/Math/UnrealMathNeon.h
@@ -121,10 +121,7 @@ FORCEINLINE VectorRegister VectorLoad( const void * Ptr )
  * @param Ptr	Unaligned memory pointer to the 3 floats
  * @return		VectorRegister(Ptr[0], Ptr[1], Ptr[2], undefined)
  */
-FORCEINLINE VectorRegister VectorLoadFloat3( const void * Ptr ) 
-{
-	return vld1q_f32( (float32_t*)Ptr );
-}
+#define VectorLoadFloat3( Ptr )			MakeVectorRegister( ((const float*)(Ptr))[0], ((const float*)(Ptr))[1], ((const float*)(Ptr))[2], 0.0f )
 
 /**
  * Loads 3 floats from unaligned memory and sets W=0.
@@ -132,10 +129,7 @@ FORCEINLINE VectorRegister VectorLoadFloat3( const void * Ptr )
  * @param Ptr	Unaligned memory pointer to the 3 floats
  * @return		VectorRegister(Ptr[0], Ptr[1], Ptr[2], 0.0f)
  */
-FORCEINLINE VectorRegister VectorLoadFloat3_W0( const void * Ptr )
-{
-	return vmulq_f32( vld1q_f32( (float32_t*)Ptr ), Vec1110 );
-}
+#define VectorLoadFloat3_W0( Ptr )		MakeVectorRegister( ((const float*)(Ptr))[0], ((const float*)(Ptr))[1], ((const float*)(Ptr))[2], 0.0f )
 
 /**
  * Loads 3 floats from unaligned memory and sets W=1.
@@ -143,10 +137,7 @@ FORCEINLINE VectorRegister VectorLoadFloat3_W0( const void * Ptr )
  * @param Ptr	Unaligned memory pointer to the 3 floats
  * @return		VectorRegister(Ptr[0], Ptr[1], Ptr[2], 1.0f)
  */
-FORCEINLINE VectorRegister VectorLoadFloat3_W1( const void * Ptr )
-{
-	return vmlaq_f32( Vec0001, vld1q_f32( (float32_t*)Ptr ), Vec1110 );
-}
+#define VectorLoadFloat3_W1( Ptr )		MakeVectorRegister( ((const float*)(Ptr))[0], ((const float*)(Ptr))[1], ((const float*)(Ptr))[2], 1.0f )
 
 /**
  * Sets a single component of a vector. Must be a define since ElementIndex needs to be a constant integer
diff --git a/Engine/Source/Runtime/Core/Public/Math/UnrealMathSSE.h b/Engine/Source/Runtime/Core/Public/Math/UnrealMathSSE.h
--- a/Engine/Source/Runtime/Core/Public/Math/UnrealMathSSE.h
+++ b/Engine/Source/Runtime/Core/Public/Math/UnrealMathSSE.h
@@ -126,7 +126,7 @@ static const VectorRegister SSE_INFINITY = MakeVectorRegister( (uint32)0x7F80000
  * @param Ptr	Unaligned memory pointer to the 3 FLOATs
  * @return		VectorRegister(Ptr[0], Ptr[1], Ptr[2], undefined)
  */
-#define VectorLoadFloat3( Ptr )			_mm_loadu_ps( (float*)(Ptr) )
+#define VectorLoadFloat3( Ptr )			MakeVectorRegister( ((const float*)(Ptr))[0], ((const float*)(Ptr))[1], ((const float*)(Ptr))[2], 0.0f )
 
 /**
  * Loads 3 FLOATs from unaligned memory and sets W=0.
@@ -134,7 +134,7 @@ static const VectorRegister SSE_INFINITY = MakeVectorRegister( (uint32)0x7F80000
  * @param Ptr	Unaligned memory pointer to the 3 FLOATs
  * @return		VectorRegister(Ptr[0], Ptr[1], Ptr[2], 0.0f)
  */
-#define VectorLoadFloat3_W0( Ptr )		_mm_and_ps( _mm_loadu_ps((float*)(Ptr)), SSE_XYZ_MASK )
+#define VectorLoadFloat3_W0( Ptr )		MakeVectorRegister( ((const float*)(Ptr))[0], ((const float*)(Ptr))[1], ((const float*)(Ptr))[2], 0.0f )
 
 /**
  * Loads 3 FLOATs from unaligned memory and sets W=1.
@@ -142,17 +142,7 @@ static const VectorRegister SSE_INFINITY = MakeVectorRegister( (uint32)0x7F80000
  * @param Ptr	Unaligned memory pointer to the 3 FLOATs
  * @return		VectorRegister(Ptr[0], Ptr[1], Ptr[2], 1.0f)
  */
-FORCEINLINE VectorRegister VectorLoadFloat3_W1( const void* Ptr )
-{
-	// Vec = (Ptr[0], Ptr[1], Ptr[2], undefined)
-	VectorRegister Vec = _mm_loadu_ps((const float*)Ptr);
-
-	// Temp = (Ptr[2], undefined, 1.0f, 1.0f)
-	VectorRegister Temp = _mm_movehl_ps( VectorOne(), Vec );
-
-	// Return (Ptr[0], Ptr[1], Ptr[2], 1.0f)
-	return _mm_shuffle_ps( Vec, Temp, SHUFFLEMASK(0,1,0,3) );
-}
+#define VectorLoadFloat3_W1( Ptr )		MakeVectorRegister( ((const float*)(Ptr))[0], ((const float*)(Ptr))[1], ((const float*)(Ptr))[2], 1.0f )
 
 /**
  * Loads 4 FLOATs from aligned memory.
