Lightmass cached approximate lighting uses the normalmap, GI from low poly meshes with a baked normal map now matches the high poly

diff --git a/Engine/Source/Programs/UnrealLightmass/Private/Lighting/TextureMapping.cpp b/Engine/Source/Programs/UnrealLightmass/Private/Lighting/TextureMapping.cpp
--- a/Engine/Source/Programs/UnrealLightmass/Private/Lighting/TextureMapping.cpp
+++ b/Engine/Source/Programs/UnrealLightmass/Private/Lighting/TextureMapping.cpp
@@ -348,8 +348,24 @@ void FStaticLightingSystem::CacheIrradiancePhotonsTextureMapping(FStaticLighting
 				MappingContext.Stats.NumCachedIrradianceSamples++;
 				FFullStaticLightingVertex CurrentVertex = TexelToVertex.GetFullVertex();
 
+				if (MaterialSettings.bUseNormalMapsForLighting && TextureMapping->Mesh->HasImportedNormal(TexelToVertex.ElementIndex))
+				{
+					const FVector4 TangentNormal = TextureMapping->Mesh->EvaluateNormal(TexelToVertex.TextureCoordinates[0], TexelToVertex.ElementIndex);
+
+					const FVector4 WorldTangentRow0(TexelToVertex.WorldTangentX.X, TexelToVertex.WorldTangentY.X, TexelToVertex.WorldTangentZ.X);
+					const FVector4 WorldTangentRow1(TexelToVertex.WorldTangentX.Y, TexelToVertex.WorldTangentY.Y, TexelToVertex.WorldTangentZ.Y);
+					const FVector4 WorldTangentRow2(TexelToVertex.WorldTangentX.Z, TexelToVertex.WorldTangentY.Z, TexelToVertex.WorldTangentZ.Z);
+					const FVector4 WorldVector(
+						Dot3(WorldTangentRow0, TangentNormal),
+						Dot3(WorldTangentRow1, TangentNormal),
+						Dot3(WorldTangentRow2, TangentNormal)
+						);
+
+					CurrentVertex.WorldTangentZ = WorldVector;
+				}
+
 				// Normalize the tangent basis and ensure it is orthonormal
-				CurrentVertex.WorldTangentZ = TexelToVertex.WorldTangentZ.UnsafeNormal3();
+				CurrentVertex.WorldTangentZ = CurrentVertex.WorldTangentZ.UnsafeNormal3();
 				CurrentVertex.TriangleNormal = TexelToVertex.TriangleNormal.UnsafeNormal3();
 				checkSlow(!CurrentVertex.TriangleNormal.ContainsNaN());
 
