Lightmass kDOP preallocates less memory, supports Nodes array resize
* LandscapeMountains Lightmass usage went from 19Gb -> 10.5Gb

diff --git a/Engine/Source/Programs/UnrealLightmass/Private/LightmassCore/Math/LMkDOP.h b/Engine/Source/Programs/UnrealLightmass/Private/LightmassCore/Math/LMkDOP.h
--- a/Engine/Source/Programs/UnrealLightmass/Private/LightmassCore/Math/LMkDOP.h
+++ b/Engine/Source/Programs/UnrealLightmass/Private/LightmassCore/Math/LMkDOP.h
@@ -545,202 +545,5 @@
 	}
 
 	/**
-	 * Determines if the node is a leaf or not. If it is not a leaf, it subdivides
-	 * the list of triangles again adding two child nodes and splitting them on
-	 * the mean (splatter method). Otherwise it sets up the triangle information.
-	 *
-	 * @param Start -- The triangle index to start processing with
-	 * @param NumTris -- The number of triangles to process
-	 * @param BuildTriangles -- The list of triangles to use for the build process
-	 * @param Nodes -- The list of nodes in this tree
-	 * @return bounding box for this node
-	 */
-	FBox SplitTriangleList(int32 Start,int32 NumTris,
-		TArray<FkDOPBuildCollisionTriangle<KDOP_IDX_TYPE> >& BuildTriangles,
-		kDOPArray<FTriangleSOA, FRangeChecklessHeapAllocator>& SOATriangles,
-		kDOPArray<NodeType, FRangeChecklessHeapAllocator>& Nodes)
-	{
-		// Figure out if we are a leaf node or not
-		if (NumTris > GKDOPMaxTrisPerLeaf)
-		{
-			// Still too many triangles, so continue subdividing the triangle list
-			bIsLeaf = 0;
-			Occupancy = 0;
-			int32 BestPlane = -1;
-			float BestMean = 0.f;
-			float BestVariance = 0.f;
-			// Determine how to split using the splatter algorithm
-			for (int32 nPlane = 0; nPlane < NUM_PLANES; nPlane++)
-			{
-				float Mean = 0.f;
-				float Variance = 0.f;
-				// Compute the mean for the triangle list
-				for (int32 nTriangle = Start; nTriangle < Start + NumTris; nTriangle++)
-				{
-					// Project the centroid of the triangle against the plane
-					// normals and accumulate to find the total projected
-					// weighting
-					Mean += BuildTriangles[nTriangle].GetCentroid()[nPlane];
-				}
-				// Divide by the number of triangles to get the average
-				Mean /= float(NumTris);
-				// Compute variance of the triangle list
-				for (int32 nTriangle = Start; nTriangle < Start + NumTris;nTriangle++)
-				{
-					// Project the centroid again
-					float Dot = BuildTriangles[nTriangle].GetCentroid()[nPlane];
-					// Now calculate the variance and accumulate it
-					Variance += (Dot - Mean) * (Dot - Mean);
-				}
-				// Get the average variance
-				Variance /= float(NumTris);
-				// Determine if this plane is the best to split on or not
-				if (Variance >= BestVariance)
-				{
-					BestPlane = nPlane;
-					BestVariance = Variance;
-					BestMean = Mean;
-				}
-			}
-			// Now that we have the plane to split on, work through the triangle
-			// list placing them on the left or right of the splitting plane
-			int32 Left = Start - 1;
-			int32 Right = Start + NumTris;
-			// Keep working through until the left index passes the right
-			while (Left < Right)
-			{
-				float Dot;
-				// Find all the triangles to the "left" of the splitting plane
-				do
-				{
-					Dot = BuildTriangles[++Left].GetCentroid()[BestPlane];
-				}
-				while (Dot < BestMean && Left < Right);
-				// Find all the triangles to the "right" of the splitting plane
-				do
-				{
-					Dot = BuildTriangles[--Right].GetCentroid()[BestPlane];
-				}
-				while (Dot >= BestMean && Right > 0 && Left < Right);
-				// Don't swap the triangle data if we just hit the end
-				if (Left < Right)
-				{
-					// Swap the triangles since they are on the wrong sides of the
-					// splitting plane
-					FkDOPBuildCollisionTriangle<KDOP_IDX_TYPE> Temp = BuildTriangles[Left];
-					BuildTriangles[Left] = BuildTriangles[Right];
-					BuildTriangles[Right] = Temp;
-				}
-			}
-			// Check for wacky degenerate case where more than GKDOPMaxTrisPerLeaf
-			// fall all in the same kDOP
-			if (Left == Start + NumTris || Right == Start)
-			{
-				Left = Start + (NumTris / 2);
-			}
-			// Add the two child nodes
-			n.LeftNode = Nodes.AddZeroed(2);
-			n.RightNode = n.LeftNode + 1;
-			// Have the left node recursively subdivide it's list and set bounding volume.
-			FBox LeftBoundingVolume = Nodes[n.LeftNode].SplitTriangleList(Start,Left - Start,BuildTriangles,SOATriangles,Nodes);
-			BoundingVolumes.SetBox(0,LeftBoundingVolume);
-			// Set unused index 2,3 to child nodes of left node.
-			BoundingVolumes.SetBox(2,Nodes[n.LeftNode].BoundingVolumes.GetBox(0));
-			BoundingVolumes.SetBox(3,Nodes[n.LeftNode].BoundingVolumes.GetBox(1));
-
-			// And now have the right node recursively subdivide it's list and set bounding volume.			
-			FBox RightBoundingVolume = Nodes[n.RightNode].SplitTriangleList(Left,Start + NumTris - Left,BuildTriangles,SOATriangles,Nodes);
-			BoundingVolumes.SetBox(1,RightBoundingVolume);
-
-			GKDOPNodes += 2;
-
-			// Non-leaf node bounds are the "sum" of the left and right nodes' volumes.
-			return LeftBoundingVolume + RightBoundingVolume;
-		}
-		else
-		{
-			// Build SOA triangles
-
-			// "NULL triangle", used when a leaf can't fill all 4 triangles in a FTriangleSOA.
-			// No line should ever hit these triangles, set the values so that it can never happen.
-			FkDOPBuildCollisionTriangle<KDOP_IDX_TYPE> EmptyTriangle(0,FVector4(0,0,0,0),FVector4(0,0,0,0),FVector4(0,0,0,0),INDEX_NONE,INDEX_NONE, false, true);
-			
-			t.StartIndex = SOATriangles.Num();
-			t.NumTriangles = Align<int32>(NumTris, 4) / 4;
-			SOATriangles.AddZeroed( t.NumTriangles );
-
-			int32 BuildTriIndex = Start;
-			for ( uint32 SOAIndex=0; SOAIndex < t.NumTriangles; ++SOAIndex )
-			{
-				FkDOPBuildCollisionTriangle<KDOP_IDX_TYPE>* Tris[4] = { &EmptyTriangle, &EmptyTriangle, &EmptyTriangle, &EmptyTriangle };
-				FTriangleSOA& SOA = SOATriangles[t.StartIndex + SOAIndex];
-				int32 SubIndex = 0;
-				for ( ; SubIndex < 4 && BuildTriIndex < (Start+NumTris); ++SubIndex, ++BuildTriIndex )
-				{
-					Tris[SubIndex] = &BuildTriangles[BuildTriIndex];
-					SOA.Payload[SubIndex] = Tris[SubIndex]->MaterialIndex;
-				}
-				for ( ; SubIndex < 4; ++SubIndex )
-				{
-					SOA.Payload[SubIndex] = 0xffffffff;
-				}
-
-				SOA.Positions[0].X = VectorSet( Tris[0]->V0.X, Tris[1]->V0.X, Tris[2]->V0.X, Tris[3]->V0.X );
-				SOA.Positions[0].Y = VectorSet( Tris[0]->V0.Y, Tris[1]->V0.Y, Tris[2]->V0.Y, Tris[3]->V0.Y );
-				SOA.Positions[0].Z = VectorSet( Tris[0]->V0.Z, Tris[1]->V0.Z, Tris[2]->V0.Z, Tris[3]->V0.Z );
-				SOA.Positions[1].X = VectorSet( Tris[0]->V1.X, Tris[1]->V1.X, Tris[2]->V1.X, Tris[3]->V1.X );
-				SOA.Positions[1].Y = VectorSet( Tris[0]->V1.Y, Tris[1]->V1.Y, Tris[2]->V1.Y, Tris[3]->V1.Y );
-				SOA.Positions[1].Z = VectorSet( Tris[0]->V1.Z, Tris[1]->V1.Z, Tris[2]->V1.Z, Tris[3]->V1.Z );
-				SOA.Positions[2].X = VectorSet( Tris[0]->V2.X, Tris[1]->V2.X, Tris[2]->V2.X, Tris[3]->V2.X );
-				SOA.Positions[2].Y = VectorSet( Tris[0]->V2.Y, Tris[1]->V2.Y, Tris[2]->V2.Y, Tris[3]->V2.Y );
-				SOA.Positions[2].Z = VectorSet( Tris[0]->V2.Z, Tris[1]->V2.Z, Tris[2]->V2.Z, Tris[3]->V2.Z );
-
-				const FVector4& Tris0LocalNormal = Tris[0]->GetLocalNormal();
-				const FVector4& Tris1LocalNormal = Tris[1]->GetLocalNormal();
-				const FVector4& Tris2LocalNormal = Tris[2]->GetLocalNormal();
-				const FVector4& Tris3LocalNormal = Tris[3]->GetLocalNormal();
-
-				SOA.Normals.X = VectorSet( Tris0LocalNormal.X, Tris1LocalNormal.X, Tris2LocalNormal.X, Tris3LocalNormal.X );
-				SOA.Normals.Y = VectorSet( Tris0LocalNormal.Y, Tris1LocalNormal.Y, Tris2LocalNormal.Y, Tris3LocalNormal.Y );
-				SOA.Normals.Z = VectorSet( Tris0LocalNormal.Z, Tris1LocalNormal.Z, Tris2LocalNormal.Z, Tris3LocalNormal.Z );
-				SOA.Normals.W = VectorSet( -Tris0LocalNormal.W, -Tris1LocalNormal.W, -Tris2LocalNormal.W, -Tris3LocalNormal.W );
-				SOA.TwoSidedMask = MakeVectorRegister(
-					(uint32)(Tris[0]->bTwoSided ? 0xFFFFFFFF : 0), 
-					(uint32)(Tris[1]->bTwoSided ? 0xFFFFFFFF : 0),
-					(uint32)(Tris[2]->bTwoSided ? 0xFFFFFFFF : 0),
-					(uint32)(Tris[3]->bTwoSided ? 0xFFFFFFFF : 0));
-				SOA.StaticAndOpaqueMask = MakeVectorRegister(
-					(uint32)(Tris[0]->bStaticAndOpaque ? 0xFFFFFFFF : 0), 
-					(uint32)(Tris[1]->bStaticAndOpaque ? 0xFFFFFFFF : 0),
-					(uint32)(Tris[2]->bStaticAndOpaque ? 0xFFFFFFFF : 0),
-					(uint32)(Tris[3]->bStaticAndOpaque ? 0xFFFFFFFF : 0));
-				SOA.MeshIndices = VectorSet(*(float*)&Tris[0]->MeshIndex, *(float*)&Tris[1]->MeshIndex, *(float*)&Tris[2]->MeshIndex, *(float*)&Tris[3]->MeshIndex);
-				SOA.LODIndices = VectorSet(*(float*)&Tris[0]->LODIndex, *(float*)&Tris[1]->LODIndex, *(float*)&Tris[2]->LODIndex, *(float*)&Tris[3]->LODIndex);
-			}
-
-			// No need to subdivide further so make this a leaf node
-			bIsLeaf = 1;
-			Occupancy = NumTris;
-			
-			// Generate bounding volume for leaf which is passed up the call chain.
-			FBox BoundingVolume(0);
-			for (int32 TriangleIndex=Start; TriangleIndex<Start + NumTris; TriangleIndex++)
-			{
-				BoundingVolume += BuildTriangles[TriangleIndex].V0;
-				BoundingVolume += BuildTriangles[TriangleIndex].V1;
-				BoundingVolume += BuildTriangles[TriangleIndex].V2;			
-			}
-			BoundingVolumes.SetBox(0,BoundingVolume);
-			BoundingVolumes.SetBox(1,BoundingVolume);
-			BoundingVolumes.SetBox(2,BoundingVolume);
-			BoundingVolumes.SetBox(3,BoundingVolume);
-
-			GKDOPTriangles += t.NumTriangles * 4;
-			GKDOPNumLeaves++;
-			return BoundingVolume;
-		}
-	}
-
-	/**
 	 * The slab testing algorithm is based on the following papers. We chose to use the
 	 * faster final hit determination, which means we'll get some false positives.
@@ -1061,17 +864,21 @@
 			FScopedRDTSCTimer kDOPBuildTimer(kDOPBuildTime);
 
 			// Empty the current set of nodes and preallocate the memory so it doesn't
-			// reallocate memory while we are recursively walking the tree
+			// reallocate memory too much while we are recursively walking the tree
 			// With near-perfect packing, we could easily size these to be
 			// Nodes = (n / 2) + 1 and SOATriangles = (n / 4) + 1
-			Nodes.Empty(BuildTriangles.Num());
-			SOATriangles.Empty(BuildTriangles.Num());
+			Nodes.Empty(BuildTriangles.Num() / 2);
+			SOATriangles.Empty(BuildTriangles.Num() / 3);
 
+			size_t NodesSize = (size_t)Nodes.GetTypeSize() * (size_t)Nodes.Max();
+			size_t SOATrianglesSize = (size_t)SOATriangles.GetTypeSize() * (size_t)SOATriangles.Max();
+			UE_LOG(LogLightmass, Log, TEXT("Preallocated %.1fGb for kDOP nodes and triangles"), (NodesSize + SOATrianglesSize) / 1024.0f / 1024.0f / 1024.0f);
+
 			// Add the root node
 			Nodes.AddZeroed();
 
 			// Now tell that node to recursively subdivide the entire set of triangles
-			Nodes[0].SplitTriangleList(0,BuildTriangles.Num(),BuildTriangles,SOATriangles,Nodes);
+			SplitTriangleList(0,0,BuildTriangles.Num(),BuildTriangles);
 
 			// Don't waste memory.
 			Nodes.Shrink();
@@ -1080,6 +887,211 @@
 		UE_LOG(LogLightmass, Log, TEXT("Building kDOP took %5.2f seconds."),kDOPBuildTime);
 	}
 
+	
 	/**
+	 * Determines if the node is a leaf or not. If it is not a leaf, it subdivides
+	 * the list of triangles again adding two child nodes and splitting them on
+	 * the mean (splatter method). Otherwise it sets up the triangle information.
+	 *
+	 * @param Start -- The triangle index to start processing with
+	 * @param NumTris -- The number of triangles to process
+	 * @param BuildTriangles -- The list of triangles to use for the build process
+	 * @param Nodes -- The list of nodes in this tree
+	 * @return bounding box for this node
+	 */
+	FBox SplitTriangleList(KDOP_IDX_TYPE NodeIndex, int32 Start, int32 NumTris, TArray<FkDOPBuildCollisionTriangle<KDOP_IDX_TYPE> >& BuildTriangles)
+	{
+		// This local node pointer will have to be updated whenever a Nodes reallocation can occur
+		NodeType* Node = &Nodes[NodeIndex];
+
+		// Figure out if we are a leaf node or not
+		if (NumTris > GKDOPMaxTrisPerLeaf)
+		{
+			// Still too many triangles, so continue subdividing the triangle list
+			Node->bIsLeaf = 0;
+			Node->Occupancy = 0;
+			int32 BestPlane = -1;
+			float BestMean = 0.f;
+			float BestVariance = 0.f;
+			// Determine how to split using the splatter algorithm
+			for (int32 nPlane = 0; nPlane < NUM_PLANES; nPlane++)
+			{
+				float Mean = 0.f;
+				float Variance = 0.f;
+				// Compute the mean for the triangle list
+				for (int32 nTriangle = Start; nTriangle < Start + NumTris; nTriangle++)
+				{
+					// Project the centroid of the triangle against the plane
+					// normals and accumulate to find the total projected
+					// weighting
+					Mean += BuildTriangles[nTriangle].GetCentroid()[nPlane];
+				}
+				// Divide by the number of triangles to get the average
+				Mean /= float(NumTris);
+				// Compute variance of the triangle list
+				for (int32 nTriangle = Start; nTriangle < Start + NumTris;nTriangle++)
+				{
+					// Project the centroid again
+					float Dot = BuildTriangles[nTriangle].GetCentroid()[nPlane];
+					// Now calculate the variance and accumulate it
+					Variance += (Dot - Mean) * (Dot - Mean);
+				}
+				// Get the average variance
+				Variance /= float(NumTris);
+				// Determine if this plane is the best to split on or not
+				if (Variance >= BestVariance)
+				{
+					BestPlane = nPlane;
+					BestVariance = Variance;
+					BestMean = Mean;
+				}
+			}
+			// Now that we have the plane to split on, work through the triangle
+			// list placing them on the left or right of the splitting plane
+			int32 Left = Start - 1;
+			int32 Right = Start + NumTris;
+			// Keep working through until the left index passes the right
+			while (Left < Right)
+			{
+				float Dot;
+				// Find all the triangles to the "left" of the splitting plane
+				do
+				{
+					Dot = BuildTriangles[++Left].GetCentroid()[BestPlane];
+				}
+				while (Dot < BestMean && Left < Right);
+				// Find all the triangles to the "right" of the splitting plane
+				do
+				{
+					Dot = BuildTriangles[--Right].GetCentroid()[BestPlane];
+				}
+				while (Dot >= BestMean && Right > 0 && Left < Right);
+				// Don't swap the triangle data if we just hit the end
+				if (Left < Right)
+				{
+					// Swap the triangles since they are on the wrong sides of the
+					// splitting plane
+					FkDOPBuildCollisionTriangle<KDOP_IDX_TYPE> Temp = BuildTriangles[Left];
+					BuildTriangles[Left] = BuildTriangles[Right];
+					BuildTriangles[Right] = Temp;
+				}
+			}
+			// Check for wacky degenerate case where more than GKDOPMaxTrisPerLeaf
+			// fall all in the same kDOP
+			if (Left == Start + NumTris || Right == Start)
+			{
+				Left = Start + (NumTris / 2);
+			}
+			// Add the two child nodes
+			KDOP_IDX_TYPE ChildIndex = Nodes.AddZeroed(2);
+			// Nodes may have resized
+			Node = &Nodes[NodeIndex];
+			Node->n.LeftNode = ChildIndex;
+			Node->n.RightNode = Node->n.LeftNode + 1;
+			// Have the left node recursively subdivide its list and set bounding volume.
+			FBox LeftBoundingVolume = SplitTriangleList(Node->n.LeftNode,Start,Left - Start,BuildTriangles);
+			// Nodes may have resized
+			Node = &Nodes[NodeIndex];
+			Node->BoundingVolumes.SetBox(0,LeftBoundingVolume);
+			// Set unused index 2,3 to child nodes of left node.
+			Node->BoundingVolumes.SetBox(2,Nodes[Node->n.LeftNode].BoundingVolumes.GetBox(0));
+			Node->BoundingVolumes.SetBox(3,Nodes[Node->n.LeftNode].BoundingVolumes.GetBox(1));
+
+			// And now have the right node recursively subdivide its list and set bounding volume.			
+			FBox RightBoundingVolume = SplitTriangleList(Node->n.RightNode,Left,Start + NumTris - Left,BuildTriangles);
+			// Nodes may have resized
+			Node = &Nodes[NodeIndex];
+			Node->BoundingVolumes.SetBox(1,RightBoundingVolume);
+
+			GKDOPNodes += 2;
+
+			// Non-leaf node bounds are the "sum" of the left and right nodes' volumes.
+			return LeftBoundingVolume + RightBoundingVolume;
+		}
+		else
+		{
+			// Build SOA triangles
+
+			// "NULL triangle", used when a leaf can't fill all 4 triangles in a FTriangleSOA.
+			// No line should ever hit these triangles, set the values so that it can never happen.
+			FkDOPBuildCollisionTriangle<KDOP_IDX_TYPE> EmptyTriangle(0,FVector4(0,0,0,0),FVector4(0,0,0,0),FVector4(0,0,0,0),INDEX_NONE,INDEX_NONE, false, true);
+			
+			Node->t.StartIndex = SOATriangles.Num();
+			Node->t.NumTriangles = Align<int32>(NumTris, 4) / 4;
+			SOATriangles.AddZeroed( Node->t.NumTriangles );
+
+			int32 BuildTriIndex = Start;
+			for ( uint32 SOAIndex=0; SOAIndex < Node->t.NumTriangles; ++SOAIndex )
+			{
+				FkDOPBuildCollisionTriangle<KDOP_IDX_TYPE>* Tris[4] = { &EmptyTriangle, &EmptyTriangle, &EmptyTriangle, &EmptyTriangle };
+				FTriangleSOA& SOA = SOATriangles[Node->t.StartIndex + SOAIndex];
+				int32 SubIndex = 0;
+				for ( ; SubIndex < 4 && BuildTriIndex < (Start+NumTris); ++SubIndex, ++BuildTriIndex )
+				{
+					Tris[SubIndex] = &BuildTriangles[BuildTriIndex];
+					SOA.Payload[SubIndex] = Tris[SubIndex]->MaterialIndex;
+				}
+				for ( ; SubIndex < 4; ++SubIndex )
+				{
+					SOA.Payload[SubIndex] = 0xffffffff;
+				}
+
+				SOA.Positions[0].X = VectorSet( Tris[0]->V0.X, Tris[1]->V0.X, Tris[2]->V0.X, Tris[3]->V0.X );
+				SOA.Positions[0].Y = VectorSet( Tris[0]->V0.Y, Tris[1]->V0.Y, Tris[2]->V0.Y, Tris[3]->V0.Y );
+				SOA.Positions[0].Z = VectorSet( Tris[0]->V0.Z, Tris[1]->V0.Z, Tris[2]->V0.Z, Tris[3]->V0.Z );
+				SOA.Positions[1].X = VectorSet( Tris[0]->V1.X, Tris[1]->V1.X, Tris[2]->V1.X, Tris[3]->V1.X );
+				SOA.Positions[1].Y = VectorSet( Tris[0]->V1.Y, Tris[1]->V1.Y, Tris[2]->V1.Y, Tris[3]->V1.Y );
+				SOA.Positions[1].Z = VectorSet( Tris[0]->V1.Z, Tris[1]->V1.Z, Tris[2]->V1.Z, Tris[3]->V1.Z );
+				SOA.Positions[2].X = VectorSet( Tris[0]->V2.X, Tris[1]->V2.X, Tris[2]->V2.X, Tris[3]->V2.X );
+				SOA.Positions[2].Y = VectorSet( Tris[0]->V2.Y, Tris[1]->V2.Y, Tris[2]->V2.Y, Tris[3]->V2.Y );
+				SOA.Positions[2].Z = VectorSet( Tris[0]->V2.Z, Tris[1]->V2.Z, Tris[2]->V2.Z, Tris[3]->V2.Z );
+
+				const FVector4& Tris0LocalNormal = Tris[0]->GetLocalNormal();
+				const FVector4& Tris1LocalNormal = Tris[1]->GetLocalNormal();
+				const FVector4& Tris2LocalNormal = Tris[2]->GetLocalNormal();
+				const FVector4& Tris3LocalNormal = Tris[3]->GetLocalNormal();
+
+				SOA.Normals.X = VectorSet( Tris0LocalNormal.X, Tris1LocalNormal.X, Tris2LocalNormal.X, Tris3LocalNormal.X );
+				SOA.Normals.Y = VectorSet( Tris0LocalNormal.Y, Tris1LocalNormal.Y, Tris2LocalNormal.Y, Tris3LocalNormal.Y );
+				SOA.Normals.Z = VectorSet( Tris0LocalNormal.Z, Tris1LocalNormal.Z, Tris2LocalNormal.Z, Tris3LocalNormal.Z );
+				SOA.Normals.W = VectorSet( -Tris0LocalNormal.W, -Tris1LocalNormal.W, -Tris2LocalNormal.W, -Tris3LocalNormal.W );
+				SOA.TwoSidedMask = MakeVectorRegister(
+					(uint32)(Tris[0]->bTwoSided ? 0xFFFFFFFF : 0), 
+					(uint32)(Tris[1]->bTwoSided ? 0xFFFFFFFF : 0),
+					(uint32)(Tris[2]->bTwoSided ? 0xFFFFFFFF : 0),
+					(uint32)(Tris[3]->bTwoSided ? 0xFFFFFFFF : 0));
+				SOA.StaticAndOpaqueMask = MakeVectorRegister(
+					(uint32)(Tris[0]->bStaticAndOpaque ? 0xFFFFFFFF : 0), 
+					(uint32)(Tris[1]->bStaticAndOpaque ? 0xFFFFFFFF : 0),
+					(uint32)(Tris[2]->bStaticAndOpaque ? 0xFFFFFFFF : 0),
+					(uint32)(Tris[3]->bStaticAndOpaque ? 0xFFFFFFFF : 0));
+				SOA.MeshIndices = VectorSet(*(float*)&Tris[0]->MeshIndex, *(float*)&Tris[1]->MeshIndex, *(float*)&Tris[2]->MeshIndex, *(float*)&Tris[3]->MeshIndex);
+				SOA.LODIndices = VectorSet(*(float*)&Tris[0]->LODIndex, *(float*)&Tris[1]->LODIndex, *(float*)&Tris[2]->LODIndex, *(float*)&Tris[3]->LODIndex);
+			}
+
+			// No need to subdivide further so make this a leaf node
+			Node->bIsLeaf = 1;
+			Node->Occupancy = NumTris;
+			
+			// Generate bounding volume for leaf which is passed up the call chain.
+			FBox BoundingVolume(0);
+			for (int32 TriangleIndex=Start; TriangleIndex<Start + NumTris; TriangleIndex++)
+			{
+				BoundingVolume += BuildTriangles[TriangleIndex].V0;
+				BoundingVolume += BuildTriangles[TriangleIndex].V1;
+				BoundingVolume += BuildTriangles[TriangleIndex].V2;			
+			}
+			Node->BoundingVolumes.SetBox(0,BoundingVolume);
+			Node->BoundingVolumes.SetBox(1,BoundingVolume);
+			Node->BoundingVolumes.SetBox(2,BoundingVolume);
+			Node->BoundingVolumes.SetBox(3,BoundingVolume);
+
+			GKDOPTriangles += Node->t.NumTriangles * 4;
+			GKDOPNumLeaves++;
+			return BoundingVolume;
+		}
+	}
+
+	/**
 	 * Figures out whether the check even hits the root node's bounding volume. If
 	 * it does, it recursively searches for a triangle to hit.
