Lightmass computes direct lighting from emissive materials on components with Static Mobility

diff --git a/Engine/Source/Programs/UnrealLightmass/Private/Lighting/FinalGather.cpp b/Engine/Source/Programs/UnrealLightmass/Private/Lighting/FinalGather.cpp
--- a/Engine/Source/Programs/UnrealLightmass/Private/Lighting/FinalGather.cpp
+++ b/Engine/Source/Programs/UnrealLightmass/Private/Lighting/FinalGather.cpp
@@ -261,6 +261,10 @@ FLinearColor FStaticLightingSystem::CalculateExitantRadiance(
 				}
 			}
 		}
+
+		//@todo - cache emissive on surfaces like direct lighting?
+		const FLinearColor EmissiveLighting = HitMesh->EvaluateEmissive(Vertex.TextureCoordinates[0], ElementIndex);
+		AccumulatedRadiance += EmissiveLighting;
 	}
 
 	// So we can compare it against FLinearColor::Black easily
@@ -527,7 +531,7 @@ FLinearColor FStaticLightingSystem::FinalGatherSample(
 					bDebugThisTexel && (!PhotonMappingSettings.bUsePhotonMapping || !PhotonMappingSettings.bVisualizePhotonImportanceSamples));
 
 				checkSlow(FLinearColorUtils::AreFloatsValid(PathVertexOutgoingRadiance));
-				Lighting = PathVertexOutgoingRadiance;
+				Lighting += PathVertexOutgoingRadiance;
 
 #if ALLOW_LIGHTMAP_SAMPLE_DEBUGGING
 				if (PathVertexOutgoingRadiance.R > DELTA || PathVertexOutgoingRadiance.G > DELTA || PathVertexOutgoingRadiance.B > DELTA)
@@ -551,7 +555,7 @@ FLinearColor FStaticLightingSystem::FinalGatherSample(
 		if (TangentPathDirection.Z > 0)
 		{
 			const FLinearColor EnvironmentLighting = EvaluateEnvironmentLighting(-WorldPathDirection);
-			Lighting = EnvironmentLighting;
+			Lighting += EnvironmentLighting;
 		}
 	}
 
