Fixed clang warnings (&& and || in the same scope) in Lightmass.

- No change in code behavior (the code was correct AFAICT).

diff --git a/Engine/Source/Programs/UnrealLightmass/Private/Lighting/Collision.cpp b/Engine/Source/Programs/UnrealLightmass/Private/Lighting/Collision.cpp
--- a/Engine/Source/Programs/UnrealLightmass/Private/Lighting/Collision.cpp
+++ b/Engine/Source/Programs/UnrealLightmass/Private/Lighting/Collision.cpp
@@ -311,13 +311,13 @@ bool FStaticLightingAggregateMesh::IntersectLightRay(
 			// Clip the ray so it is now from the intersection point (plus some amount to avoid intersecting the same triangle again) to the original end of the ray
 			ClippedLightRay.ClipAgainstIntersectionFromEnd(ClosestIntersection.IntersectionVertex.WorldPosition + Scene.SceneConstants.VisibilityRayOffsetDistance * ClippedLightRay.Direction.UnsafeNormal3());
 			
-			if (ClosestIntersection.Mesh == LightRay.Mesh && ((ClosestIntersection.Mesh->LightingFlags & GI_INSTANCE_SELFSHADOWDISABLE) || (LightRay.TraceFlags & LIGHTRAY_SELFSHADOWDISABLE))
-				|| bDirectShadowingRay && ClosestIntersection.Mesh->IsIndirectlyShadowedOnly(ClosestIntersection.ElementIndex)
-				|| ClosestIntersection.Mesh != LightRay.Mesh && (ClosestIntersection.Mesh->LightingFlags & GI_INSTANCE_SELFSHADOWONLY))
+			if ((ClosestIntersection.Mesh == LightRay.Mesh && ((ClosestIntersection.Mesh->LightingFlags & GI_INSTANCE_SELFSHADOWDISABLE) || (LightRay.TraceFlags & LIGHTRAY_SELFSHADOWDISABLE)))
+				|| (bDirectShadowingRay && ClosestIntersection.Mesh->IsIndirectlyShadowedOnly(ClosestIntersection.ElementIndex))
+				|| (ClosestIntersection.Mesh != LightRay.Mesh && (ClosestIntersection.Mesh->LightingFlags & GI_INSTANCE_SELFSHADOWONLY)))
 			{
 				// Nothing more to do if we're just avoiding shadowing based on trace flags, just keep going
 			}
-			else if (ClosestIntersection.Mesh->IsMasked(ClosestIntersection.ElementIndex) || bDirectShadowingRay && ClosestIntersection.Mesh->IsCastingShadowsAsMasked(ClosestIntersection.ElementIndex))
+			else if (ClosestIntersection.Mesh->IsMasked(ClosestIntersection.ElementIndex) || (bDirectShadowingRay && ClosestIntersection.Mesh->IsCastingShadowsAsMasked(ClosestIntersection.ElementIndex)))
 			{
 				// look to see if we hit a hole or opaque part
  				if (ClosestIntersection.Mesh->EvaluateMaskedCollision(ClosestIntersection.IntersectionVertex.TextureCoordinates[0], ClosestIntersection.ElementIndex))
@@ -439,10 +439,10 @@ bool FStaticLightingAggregateMesh::IntersectLightRay(
 	while (ClosestIntersection.bIntersects 
 		&& (ClosestIntersection.Mesh->IsTranslucent(ClosestIntersection.ElementIndex) ||
 			ClosestIntersection.Mesh->IsMasked(ClosestIntersection.ElementIndex) ||
-			ClosestIntersection.Mesh == LightRay.Mesh && ((ClosestIntersection.Mesh->LightingFlags & GI_INSTANCE_SELFSHADOWDISABLE) || (LightRay.TraceFlags & LIGHTRAY_SELFSHADOWDISABLE)) ||
+			(ClosestIntersection.Mesh == LightRay.Mesh && ((ClosestIntersection.Mesh->LightingFlags & GI_INSTANCE_SELFSHADOWDISABLE) || (LightRay.TraceFlags & LIGHTRAY_SELFSHADOWDISABLE))) ||
 			// Continue tracing if we are only allowed to self shadow and intersected a different mesh
-			ClosestIntersection.Mesh != LightRay.Mesh && (ClosestIntersection.Mesh->LightingFlags & GI_INSTANCE_SELFSHADOWONLY) ||
-			bDirectShadowingRay && ClosestIntersection.Mesh->IsIndirectlyShadowedOnly(ClosestIntersection.ElementIndex))
+			(ClosestIntersection.Mesh != LightRay.Mesh && (ClosestIntersection.Mesh->LightingFlags & GI_INSTANCE_SELFSHADOWONLY)) ||
+			(bDirectShadowingRay && ClosestIntersection.Mesh->IsIndirectlyShadowedOnly(ClosestIntersection.ElementIndex)))
 		&& NumIterativeIntersections < MaxNumIterativeIntersections);
 
 	if (NumIterativeIntersections >= MaxNumIterativeIntersections)
diff --git a/Engine/Source/Programs/UnrealLightmass/Private/Lighting/FinalGather.cpp b/Engine/Source/Programs/UnrealLightmass/Private/Lighting/FinalGather.cpp
--- a/Engine/Source/Programs/UnrealLightmass/Private/Lighting/FinalGather.cpp
+++ b/Engine/Source/Programs/UnrealLightmass/Private/Lighting/FinalGather.cpp
@@ -779,7 +779,7 @@ public:
 
 						// Refine if the importance cone threshold is exceeded or there was a big enough brightness difference
 						if (MaxRelativeDifference > BrightnessThreshold 
-							|| bRefineForSkyOcclusion && MaxSkyOcclusionDifference > SkyOcclusionThreshold
+							|| (bRefineForSkyOcclusion && MaxSkyOcclusionDifference > SkyOcclusionThreshold)
 							|| bSuperSampleDueToImportanceCones)
 						{
 							FSimpleQuadTreeNode<FRefinementElement>* Node = &Cells[ThetaIndex * NumPhiSteps + PhiIndex].RootNode;
@@ -876,7 +876,7 @@ public:
 
 							// Refine if the importance cone threshold is exceeded or there was a big enough brightness difference
 							if (MaxRelativeDifference > BrightnessThreshold 
-								|| bRefineForSkyOcclusion && MaxSkyOcclusionDifference > SkyOcclusionThreshold
+								|| (bRefineForSkyOcclusion && MaxSkyOcclusionDifference > SkyOcclusionThreshold)
 								|| bSuperSampleDueToImportanceCones)
 							{
 								NextNodesToRefine->Add(FRefinementTraversalContext(
diff --git a/Engine/Source/Programs/UnrealLightmass/Private/Lighting/PhotonMapping.cpp b/Engine/Source/Programs/UnrealLightmass/Private/Lighting/PhotonMapping.cpp
--- a/Engine/Source/Programs/UnrealLightmass/Private/Lighting/PhotonMapping.cpp
+++ b/Engine/Source/Programs/UnrealLightmass/Private/Lighting/PhotonMapping.cpp
@@ -552,7 +552,7 @@ void FStaticLightingSystem::EmitDirectPhotonsWorkRange(
 			// Apply transmission
 			PathAlpha *= PathIntersection.Transmission;
 
-			if (PathAlpha.R < DELTA && PathAlpha.G < DELTA && PathAlpha.B < DELTA
+			if ((PathAlpha.R < DELTA && PathAlpha.G < DELTA && PathAlpha.B < DELTA)
 				// Ray can hit translucent meshes if they have bCastShadowAsMasked, but we don't have diffuse for translucency, so just terminate
 				|| PathIntersection.Mesh->IsTranslucent(PathIntersection.ElementIndex))
 			{
@@ -979,8 +979,8 @@ void FStaticLightingSystem::EmitIndirectPhotonsWorkRange(
 						// Only allow creating an irradiance photon if one or more indirect bounces are required
 						// The final gather is the first bounce when enabled
 						bShouldCreateIrradiancePhoton = 
-							PhotonMappingSettings.bUseFinalGathering && GeneralSettings.NumIndirectLightingBounces > 1
-							|| !PhotonMappingSettings.bUseFinalGathering && GeneralSettings.NumIndirectLightingBounces > 0;
+							(PhotonMappingSettings.bUseFinalGathering && GeneralSettings.NumIndirectLightingBounces > 1)
+							|| (!PhotonMappingSettings.bUseFinalGathering && GeneralSettings.NumIndirectLightingBounces > 0);
 					}
 					else
 					{
@@ -989,8 +989,8 @@ void FStaticLightingSystem::EmitIndirectPhotonsWorkRange(
 						// Only allow creating an irradiance photon if two or more indirect bounces are required
 						// The final gather is the first bounce when enabled
 						bShouldCreateIrradiancePhoton = 
-							PhotonMappingSettings.bUseFinalGathering && GeneralSettings.NumIndirectLightingBounces > 2
-							|| !PhotonMappingSettings.bUseFinalGathering && GeneralSettings.NumIndirectLightingBounces > 1;
+							(PhotonMappingSettings.bUseFinalGathering && GeneralSettings.NumIndirectLightingBounces > 2)
+							|| (!PhotonMappingSettings.bUseFinalGathering && GeneralSettings.NumIndirectLightingBounces > 1);
 					}
 
 #if ALLOW_LIGHTMAP_SAMPLE_DEBUGGING
@@ -1225,7 +1225,7 @@ void FStaticLightingSystem::MarkIrradiancePhotonsWorkRange(
 		FIrradiancePhoton& CurrentIrradiancePhoton = CurrentArray[PhotonIndex];
 
 		// Only add direct contribution if we are final gathering and at least one bounce is required,
-		if (PhotonMappingSettings.bUseFinalGathering && GeneralSettings.NumIndirectLightingBounces > 0
+		if ((PhotonMappingSettings.bUseFinalGathering && GeneralSettings.NumIndirectLightingBounces > 0)
 			// Or if photon mapping is being used for direct lighting.
 			|| PhotonMappingSettings.bVisualizeCachedApproximateDirectLighting)
 		{
@@ -1406,7 +1406,7 @@ void FStaticLightingSystem::CalculateIrradiancePhotonsWorkRange(
 
 		FLinearColor AccumulatedIrradiance(FLinearColor::Black);
 		// Only add direct contribution if we are final gathering and at least one bounce is required,
-		if ((PhotonMappingSettings.bUseFinalGathering && GeneralSettings.NumIndirectLightingBounces > 0
+		if (((PhotonMappingSettings.bUseFinalGathering && GeneralSettings.NumIndirectLightingBounces > 0)
 			// Or if photon mapping is being used for direct lighting.
 			|| PhotonMappingSettings.bVisualizeCachedApproximateDirectLighting)
 			&& PhotonMappingSettings.bUsePhotonDirectLightingInFinalGather)
@@ -1425,17 +1425,17 @@ void FStaticLightingSystem::CalculateIrradiancePhotonsWorkRange(
 
 			// Only add direct contribution if it should be viewed given ViewSingleBounceNumber
 			if (GeneralSettings.ViewSingleBounceNumber < 0
-				|| PhotonMappingSettings.bUseFinalGathering && GeneralSettings.ViewSingleBounceNumber == 1
-				|| !PhotonMappingSettings.bUseFinalGathering && GeneralSettings.ViewSingleBounceNumber == 0
-				|| PhotonMappingSettings.bVisualizeCachedApproximateDirectLighting && GeneralSettings.ViewSingleBounceNumber == 0)
+				|| (PhotonMappingSettings.bUseFinalGathering && GeneralSettings.ViewSingleBounceNumber == 1)
+				|| (!PhotonMappingSettings.bUseFinalGathering && GeneralSettings.ViewSingleBounceNumber == 0)
+				|| (PhotonMappingSettings.bVisualizeCachedApproximateDirectLighting && GeneralSettings.ViewSingleBounceNumber == 0))
 			{
 				AccumulatedIrradiance = DirectPhotonIrradiance;
 			}
 		}
 
 		// If we are final gathering, first bounce photons are actually the second lighting bounce since the final gather is the first bounce
-		if (PhotonMappingSettings.bUseFinalGathering && GeneralSettings.NumIndirectLightingBounces > 1
-			|| !PhotonMappingSettings.bUseFinalGathering && GeneralSettings.NumIndirectLightingBounces > 0)
+		if ((PhotonMappingSettings.bUseFinalGathering && GeneralSettings.NumIndirectLightingBounces > 1)
+			|| (!PhotonMappingSettings.bUseFinalGathering && GeneralSettings.NumIndirectLightingBounces > 0))
 		{
 			const FLinearColor FirstBouncePhotonIrradiance = CalculatePhotonIrradiance(
 				FirstBouncePhotonMap, 
@@ -1451,15 +1451,15 @@ void FStaticLightingSystem::CalculateIrradiancePhotonsWorkRange(
 
 			// Only add first bounce contribution if it should be viewed given ViewSingleBounceNumber
 			if (GeneralSettings.ViewSingleBounceNumber < 0 
-				|| PhotonMappingSettings.bUseFinalGathering && GeneralSettings.ViewSingleBounceNumber == 2
-				|| !PhotonMappingSettings.bUseFinalGathering && GeneralSettings.ViewSingleBounceNumber == 1)
+				|| (PhotonMappingSettings.bUseFinalGathering && GeneralSettings.ViewSingleBounceNumber == 2)
+				|| (!PhotonMappingSettings.bUseFinalGathering && GeneralSettings.ViewSingleBounceNumber == 1))
 			{
 				AccumulatedIrradiance += FirstBouncePhotonIrradiance;
 			}
 
 			// If we are final gathering, second bounce photons are actually the third lighting bounce since the final gather is the first bounce
-			if (PhotonMappingSettings.bUseFinalGathering && GeneralSettings.NumIndirectLightingBounces > 2
-				|| !PhotonMappingSettings.bUseFinalGathering && GeneralSettings.NumIndirectLightingBounces > 1)
+			if ((PhotonMappingSettings.bUseFinalGathering && GeneralSettings.NumIndirectLightingBounces > 2)
+				|| (!PhotonMappingSettings.bUseFinalGathering && GeneralSettings.NumIndirectLightingBounces > 1))
 			{
 				const FLinearColor SecondBouncePhotonIrradiance = CalculatePhotonIrradiance(
 					SecondBouncePhotonMap, 
@@ -1475,8 +1475,8 @@ void FStaticLightingSystem::CalculateIrradiancePhotonsWorkRange(
 
 				// Only add second and up bounce contribution if it should be viewed given ViewSingleBounceNumber
 				if (GeneralSettings.ViewSingleBounceNumber < 0 
-					|| PhotonMappingSettings.bUseFinalGathering && GeneralSettings.ViewSingleBounceNumber == 3
-					|| !PhotonMappingSettings.bUseFinalGathering && GeneralSettings.ViewSingleBounceNumber == 2)
+					|| (PhotonMappingSettings.bUseFinalGathering && GeneralSettings.ViewSingleBounceNumber == 3)
+					|| (!PhotonMappingSettings.bUseFinalGathering && GeneralSettings.ViewSingleBounceNumber == 2))
 				{
 					AccumulatedIrradiance += SecondBouncePhotonIrradiance;
 				}
@@ -2038,8 +2038,8 @@ FIrradiancePhoton* FStaticLightingSystem::FindNearestIrradiancePhoton(
 				// Only searching for irradiance photons with normals similar to the search normal
 				if (CosTheta > PhotonMappingSettings.PhotonSearchAngleThreshold
 					// And closer to the search position than the max search distance.
-					&& (CurrentPhoton.HasDirectContribution() && (DistanceSquared < FMath::Square(PhotonMappingSettings.DirectPhotonSearchDistance))
-					|| !CurrentPhoton.HasDirectContribution() && (DistanceSquared < FMath::Square(PhotonMappingSettings.IndirectPhotonSearchDistance))))
+					&& ((CurrentPhoton.HasDirectContribution() && (DistanceSquared < FMath::Square(PhotonMappingSettings.DirectPhotonSearchDistance)))
+					|| (!CurrentPhoton.HasDirectContribution() && (DistanceSquared < FMath::Square(PhotonMappingSettings.IndirectPhotonSearchDistance)))))
 				{
 					// Only accept irradiance photons within an angle of the plane defined by the vertex normal
 					// This avoids expensive visibility traces to photons that are probably not on the same surface
diff --git a/Engine/Source/Programs/UnrealLightmass/Private/Lighting/PrecomputedVisibility.cpp b/Engine/Source/Programs/UnrealLightmass/Private/Lighting/PrecomputedVisibility.cpp
--- a/Engine/Source/Programs/UnrealLightmass/Private/Lighting/PrecomputedVisibility.cpp
+++ b/Engine/Source/Programs/UnrealLightmass/Private/Lighting/PrecomputedVisibility.cpp
@@ -228,7 +228,7 @@ void FStaticLightingSystem::SetupPrecomputedVisibility()
 
 					// Only place cells on opaque surfaces if requested, which can save some memory for foliage maps
 					if (!PrecomputedVisibilitySettings.bPlaceCellsOnOpaqueOnly 
-						|| !CurrentMesh->IsMasked(ElementIndex) && !CurrentMesh->IsTranslucent(ElementIndex))
+						|| (!CurrentMesh->IsMasked(ElementIndex) && !CurrentMesh->IsTranslucent(ElementIndex)))
 					{
 						FVector2D XYPositions[3];
 						for (int32 VertIndex = 0; VertIndex < 3; VertIndex++)
@@ -313,9 +313,9 @@ void FStaticLightingSystem::SetupPrecomputedVisibility()
 					// Place a new cell if this is the highest height
 					if (HeightIndex + 1 == Cell.HitTriangles.Num()
 						// Or if there's a gap above this height of size PlayAreaHeight
-						|| (Cell.HitTriangles[HeightIndex + 1].HeightRange.Y - CurrentMaxHeight) > PrecomputedVisibilitySettings.PlayAreaHeight
+						|| ((Cell.HitTriangles[HeightIndex + 1].HeightRange.Y - CurrentMaxHeight) > PrecomputedVisibilitySettings.PlayAreaHeight
 						// And this height is not within a cell that was just placed
-						&& CurrentMaxHeight - LastSampleHeight > PrecomputedVisibilitySettings.PlayAreaHeight)
+						&& CurrentMaxHeight - LastSampleHeight > PrecomputedVisibilitySettings.PlayAreaHeight))
 					{
 						FPrecomputedVisibilityCell NewCell;
 						NewCell.Bounds = FBox(
diff --git a/Engine/Source/Programs/UnrealLightmass/Private/Lighting/VolumeDistanceField.cpp b/Engine/Source/Programs/UnrealLightmass/Private/Lighting/VolumeDistanceField.cpp
--- a/Engine/Source/Programs/UnrealLightmass/Private/Lighting/VolumeDistanceField.cpp
+++ b/Engine/Source/Programs/UnrealLightmass/Private/Lighting/VolumeDistanceField.cpp
@@ -133,7 +133,7 @@ void FStaticLightingSystem::CalculateVolumeDistanceFieldWorkRange(int32 ZIndex)
 
 					if (Intersection.bIntersects)
 					{
-						if ((i == 0 && FMath::Abs(Intersection.IntersectionVertex.WorldTangentZ.Z) >= .707f || i == 1 && FMath::Abs(Intersection.IntersectionVertex.WorldTangentZ.Z) < .707f))
+						if ((i == 0 && FMath::Abs(Intersection.IntersectionVertex.WorldTangentZ.Z) >= .707f) || (i == 1 && FMath::Abs(Intersection.IntersectionVertex.WorldTangentZ.Z) < .707f))
 						{
 							Hit[i]++;
 							if (Dot3(Ray.Direction, Intersection.IntersectionVertex.WorldTangentZ) < 0)
