build script, cross compile unreallightmass

diff --git a/Engine/Source/Programs/UnrealLightmass/UnrealLightmass.Target.cs b/Engine/Source/Programs/UnrealLightmass/UnrealLightmass.Target.cs
--- a/Engine/Source/Programs/UnrealLightmass/UnrealLightmass.Target.cs
+++ b/Engine/Source/Programs/UnrealLightmass/UnrealLightmass.Target.cs
@@ -81,4 +81,14 @@ public class UnrealLightmassTarget : TargetRules
         SeparateNode = false;
         return true;
     }
+	public override List<UnrealTargetPlatform> GUBP_ToolPlatforms(UnrealTargetPlatform InHostPlatform)
+	{
+		List<UnrealTargetPlatform> PlatformList = new List<UnrealTargetPlatform>();
+		PlatformList.Add(InHostPlatform);
+		if(InHostPlatform == UnrealTargetPlatform.Win64)
+		{
+			PlatformList.Add(UnrealTargetPlatform.Linux);			
+		}
+		return PlatformList;
+	}
 }
