Fix build warnings due to deprecated functions.

diff --git a/Engine/Source/Programs/UnrealLightmass/Private/Lighting/LightingMesh.h b/Engine/Source/Programs/UnrealLightmass/Private/Lighting/LightingMesh.h
--- a/Engine/Source/Programs/UnrealLightmass/Private/Lighting/LightingMesh.h
+++ b/Engine/Source/Programs/UnrealLightmass/Private/Lighting/LightingMesh.h
@@ -48,11 +48,11 @@ struct FStaticLightingVertex: public FStaticLightingVertexData
 		if (TangentXCandidate.SizeSquared3() < KINDA_SMALL_NUMBER)
 		{
 			// The normal was nearly equal to the Y axis, use the X axis instead
-			WorldTangentX = (WorldTangentZ ^ FVector4(1,0,0)).UnsafeNormal3();
+			WorldTangentX = (WorldTangentZ ^ FVector4(1,0,0)).GetUnsafeNormal3();
 		}
 		else
 		{
-			WorldTangentX = TangentXCandidate.UnsafeNormal3();
+			WorldTangentX = TangentXCandidate.GetUnsafeNormal3();
 		}
 		WorldTangentY = WorldTangentZ ^ WorldTangentX;
 		checkSlow(WorldTangentY.IsUnit3());
@@ -166,11 +166,11 @@ struct FFullStaticLightingVertex : public FStaticLightingVertex
 		if (TangentXCandidate.SizeSquared3() < KINDA_SMALL_NUMBER)
 		{
 			// The normal was nearly equal to the Y axis, use the X axis instead
-			TriangleTangentX = (TriangleNormal ^ FVector4(1,0,0)).UnsafeNormal3();
+			TriangleTangentX = (TriangleNormal ^ FVector4(1,0,0)).GetUnsafeNormal3();
 		}
 		else
 		{
-			TriangleTangentX = TangentXCandidate.UnsafeNormal3();
+			TriangleTangentX = TangentXCandidate.GetUnsafeNormal3();
 		}
 		TriangleTangentY = TriangleNormal ^ TriangleTangentX;
 		checkSlow(TriangleTangentY.IsUnit3());
diff --git a/Engine/Source/Programs/UnrealLightmass/Private/Lighting/MonteCarlo.cpp b/Engine/Source/Programs/UnrealLightmass/Private/Lighting/MonteCarlo.cpp
--- a/Engine/Source/Programs/UnrealLightmass/Private/Lighting/MonteCarlo.cpp
+++ b/Engine/Source/Programs/UnrealLightmass/Private/Lighting/MonteCarlo.cpp
@@ -15,11 +15,11 @@ void GenerateCoordinateSystem(const FVector4& ZAxis, FVector4& XAxis, FVector4&
 	if (XAxisCandidate.SizeSquared3() < KINDA_SMALL_NUMBER)
 	{
 		// The vector was nearly equal to the Y axis, use the X axis instead
-		XAxis = (ZAxis ^ FVector4(1,0,0)).UnsafeNormal3();
+		XAxis = (ZAxis ^ FVector4(1,0,0)).GetUnsafeNormal3();
 	}
 	else
 	{
-		XAxis = XAxisCandidate.UnsafeNormal3();
+		XAxis = XAxisCandidate.GetUnsafeNormal3();
 	}
 
 	YAxis = ZAxis ^ XAxis;
@@ -29,7 +29,7 @@ void GenerateCoordinateSystem(const FVector4& ZAxis, FVector4& XAxis, FVector4&
 /** Generates a pseudo-random unit vector, uniformly distributed over all directions. */
 FVector4 GetUnitVector(FLMRandomStream& RandomStream)
 {
-	return GetUnitPosition(RandomStream).UnsafeNormal3();
+	return GetUnitPosition(RandomStream).GetUnsafeNormal3();
 }
 
 /** Generates a pseudo-random position inside the unit sphere, uniformly distributed over the volume of the sphere. */
@@ -99,15 +99,15 @@ FVector4 GetModifiedPhongSpecularVector(FLMRandomStream& RandomStream, const FVe
 		const FVector4 GeneratedSpecularTangentVector(FMath::Cos(Phi) * SinTheta, FMath::Sin(Phi) * SinTheta, FMath::Cos(Alpha));
 
 		// Generate the X and Y axes of the coordinate space whose Z is the perfect specular direction
-		FVector4 SpecularTangentX = (TangentSpecularDirection ^ FVector4(0,1,0)).UnsafeNormal3();
+		FVector4 SpecularTangentX = (TangentSpecularDirection ^ FVector4(0,1,0)).GetUnsafeNormal3();
 		if (SpecularTangentX.SizeSquared3() < KINDA_SMALL_NUMBER)
 		{
 			// The specular direction was nearly equal to the Y axis, use the X axis instead
-			SpecularTangentX = (TangentSpecularDirection ^ FVector4(1,0,0)).UnsafeNormal3();
+			SpecularTangentX = (TangentSpecularDirection ^ FVector4(1,0,0)).GetUnsafeNormal3();
 		}
 		else
 		{
-			SpecularTangentX = SpecularTangentX.UnsafeNormal3();
+			SpecularTangentX = SpecularTangentX.GetUnsafeNormal3();
 		}
 		const FVector4 SpecularTangentY = TangentSpecularDirection ^ SpecularTangentX;
 
diff --git a/Engine/Source/Programs/UnrealLightmass/Private/Lighting/PhotonMapping.cpp b/Engine/Source/Programs/UnrealLightmass/Private/Lighting/PhotonMapping.cpp
--- a/Engine/Source/Programs/UnrealLightmass/Private/Lighting/PhotonMapping.cpp
+++ b/Engine/Source/Programs/UnrealLightmass/Private/Lighting/PhotonMapping.cpp
@@ -537,7 +537,7 @@ void FStaticLightingSystem::EmitDirectPhotonsWorkRange(
 		AggregateMesh.IntersectLightRay(SampleRay, true, true, true, CoherentRayCache, PathIntersection);
 		Output.DirectPhotonsTracingThreadTime += CoherentRayCache.FirstHitRayTraceTime - BeforeDirectTraceTime;
 
-		const FVector4 WorldPathDirection = SampleRay.Direction.UnsafeNormal3();
+		const FVector4 WorldPathDirection = SampleRay.Direction.GetUnsafeNormal3();
 
 		// Register this photon path as long as it hit a frontface of something in the scene
 		if (PathIntersection.bIntersects 
@@ -937,7 +937,7 @@ void FStaticLightingSystem::EmitIndirectPhotonsWorkRange(
 		Output.IntersectLightRayThreadTime += CoherentRayCache.FirstHitRayTraceTime - BeforeLightRayTime;
 
 		LIGHTINGSTAT(FScopedRDTSCTimer PhotonTracingTimer(Output.PhotonBounceTracingThreadTime));
-		FVector4 WorldPathDirection = SampleRay.Direction.UnsafeNormal3();
+		FVector4 WorldPathDirection = SampleRay.Direction.GetUnsafeNormal3();
 		// Continue tracing this photon path as long as it hit a frontface of something in the scene
 		while (PathIntersection.bIntersects && Dot3(WorldPathDirection, PathIntersection.IntersectionVertex.WorldTangentZ) < 0.0f)
 		{
diff --git a/Engine/Source/Programs/UnrealLightmass/Private/Lighting/TextureMapping.cpp b/Engine/Source/Programs/UnrealLightmass/Private/Lighting/TextureMapping.cpp
--- a/Engine/Source/Programs/UnrealLightmass/Private/Lighting/TextureMapping.cpp
+++ b/Engine/Source/Programs/UnrealLightmass/Private/Lighting/TextureMapping.cpp
@@ -365,14 +365,14 @@ void FStaticLightingSystem::CacheIrradiancePhotonsTextureMapping(FStaticLighting
 				}
 
 				// Normalize the tangent basis and ensure it is orthonormal
-				CurrentVertex.WorldTangentZ = CurrentVertex.WorldTangentZ.UnsafeNormal3();
-				CurrentVertex.TriangleNormal = TexelToVertex.TriangleNormal.UnsafeNormal3();
+				CurrentVertex.WorldTangentZ = CurrentVertex.WorldTangentZ.GetUnsafeNormal3();
+				CurrentVertex.TriangleNormal = TexelToVertex.TriangleNormal.GetUnsafeNormal3();
 				checkSlow(!CurrentVertex.TriangleNormal.ContainsNaN());
 
 				const FVector4 OriginalTangentX = CurrentVertex.WorldTangentX;
 				const FVector4 OriginalTangentY = CurrentVertex.WorldTangentY;
 
-				CurrentVertex.WorldTangentY = (CurrentVertex.WorldTangentZ ^ CurrentVertex.WorldTangentX).UnsafeNormal3();
+				CurrentVertex.WorldTangentY = (CurrentVertex.WorldTangentZ ^ CurrentVertex.WorldTangentX).GetUnsafeNormal3();
 				// Maintain handedness
 				if (Dot3(CurrentVertex.WorldTangentY, OriginalTangentY) < 0)
 				{
@@ -956,14 +956,14 @@ void FStaticLightingSystem::SetupTextureMapping(
 				}
 
 				// Normalize the tangent basis and ensure it is orthonormal
-				TexelToVertex.WorldTangentZ = TexelToVertex.WorldTangentZ.UnsafeNormal3();
-				TexelToVertex.TriangleNormal = TexelToVertex.TriangleNormal.UnsafeNormal3();
+				TexelToVertex.WorldTangentZ = TexelToVertex.WorldTangentZ.GetUnsafeNormal3();
+				TexelToVertex.TriangleNormal = TexelToVertex.TriangleNormal.GetUnsafeNormal3();
 				checkSlow(!TexelToVertex.TriangleNormal.ContainsNaN());
 
 				const FVector4 OriginalTangentX = TexelToVertex.WorldTangentX;
 				const FVector4 OriginalTangentY = TexelToVertex.WorldTangentY;
 				
-				TexelToVertex.WorldTangentY = (TexelToVertex.WorldTangentZ ^ TexelToVertex.WorldTangentX).UnsafeNormal3();
+				TexelToVertex.WorldTangentY = (TexelToVertex.WorldTangentZ ^ TexelToVertex.WorldTangentX).GetUnsafeNormal3();
 				// Maintain handedness
 				if (Dot3(TexelToVertex.WorldTangentY, OriginalTangentY) < 0)
 				{
