Fix more warnings in Lightmass due to deprecated functions.

diff --git a/Engine/Source/Programs/UnrealLightmass/Private/ImportExport/Exporter.cpp b/Engine/Source/Programs/UnrealLightmass/Private/ImportExport/Exporter.cpp
--- a/Engine/Source/Programs/UnrealLightmass/Private/ImportExport/Exporter.cpp
+++ b/Engine/Source/Programs/UnrealLightmass/Private/ImportExport/Exporter.cpp
@@ -119,7 +119,7 @@ namespace Lightmass
 				}
 				if (AverageNormal.SizeSquared3() > KINDA_SMALL_NUMBER)
 				{
-					AverageNormal = AverageNormal.UnsafeNormal3();
+					AverageNormal = AverageNormal.GetUnsafeNormal3();
 				}
 				else
 				{
diff --git a/Engine/Source/Programs/UnrealLightmass/Private/ImportExport/LightmassScene.cpp b/Engine/Source/Programs/UnrealLightmass/Private/ImportExport/LightmassScene.cpp
--- a/Engine/Source/Programs/UnrealLightmass/Private/ImportExport/LightmassScene.cpp
+++ b/Engine/Source/Programs/UnrealLightmass/Private/ImportExport/LightmassScene.cpp
@@ -1471,7 +1471,7 @@ void FMeshLightPrimitive::AddSubPrimitive(const FTexelToCorners& TexelToCorners,
 
 void FMeshLightPrimitive::Finalize()
 {
-	SurfaceNormal = SurfaceNormal.SizeSquared3() > SMALL_NUMBER ? SurfaceNormal.UnsafeNormal3() : FVector4(0, 0, 1);
+	SurfaceNormal = SurfaceNormal.SizeSquared3() > SMALL_NUMBER ? SurfaceNormal.GetUnsafeNormal3() : FVector4(0, 0, 1);
 }
 
 //----------------------------------------------------------------------------
diff --git a/Engine/Source/Programs/UnrealLightmass/Private/Lighting/Collision.cpp b/Engine/Source/Programs/UnrealLightmass/Private/Lighting/Collision.cpp
--- a/Engine/Source/Programs/UnrealLightmass/Private/Lighting/Collision.cpp
+++ b/Engine/Source/Programs/UnrealLightmass/Private/Lighting/Collision.cpp
@@ -309,7 +309,7 @@ bool FStaticLightingAggregateMesh::IntersectLightRay(
 		{
 			NumIterativeIntersections++;
 			// Clip the ray so it is now from the intersection point (plus some amount to avoid intersecting the same triangle again) to the original end of the ray
-			ClippedLightRay.ClipAgainstIntersectionFromEnd(ClosestIntersection.IntersectionVertex.WorldPosition + Scene.SceneConstants.VisibilityRayOffsetDistance * ClippedLightRay.Direction.UnsafeNormal3());
+			ClippedLightRay.ClipAgainstIntersectionFromEnd(ClosestIntersection.IntersectionVertex.WorldPosition + Scene.SceneConstants.VisibilityRayOffsetDistance * ClippedLightRay.Direction.GetUnsafeNormal3());
 			
 			if ((ClosestIntersection.Mesh == LightRay.Mesh && ((ClosestIntersection.Mesh->LightingFlags & GI_INSTANCE_SELFSHADOWDISABLE) || (LightRay.TraceFlags & LIGHTRAY_SELFSHADOWDISABLE)))
 				|| (bDirectShadowingRay && ClosestIntersection.Mesh->IsIndirectlyShadowedOnly(ClosestIntersection.ElementIndex))
diff --git a/Engine/Source/Programs/UnrealLightmass/Private/Lighting/FinalGather.cpp b/Engine/Source/Programs/UnrealLightmass/Private/Lighting/FinalGather.cpp
--- a/Engine/Source/Programs/UnrealLightmass/Private/Lighting/FinalGather.cpp
+++ b/Engine/Source/Programs/UnrealLightmass/Private/Lighting/FinalGather.cpp
@@ -1656,7 +1656,7 @@ FFinalGatherSample FStaticLightingSystem::CachePointIncomingRadiance(
 						// FindNearbyPhotons only returns photons whose incident directions lie in this hemisphere, but the recalculated direction might not.
 						if (Dot3(NewDirection, Vertex.TriangleNormal) > 0.0f)
 						{
-							ImportancePhotonDirections.Add(NewDirection.UnsafeNormal3());
+							ImportancePhotonDirections.Add(NewDirection.GetUnsafeNormal3());
 						}
 					}
 				}
diff --git a/Engine/Source/Programs/UnrealLightmass/Private/Lighting/LightingMesh.cpp b/Engine/Source/Programs/UnrealLightmass/Private/Lighting/LightingMesh.cpp
--- a/Engine/Source/Programs/UnrealLightmass/Private/Lighting/LightingMesh.cpp
+++ b/Engine/Source/Programs/UnrealLightmass/Private/Lighting/LightingMesh.cpp
@@ -261,7 +261,7 @@ void FStaticLightingMesh::CreateMeshAreaLights(
 					{
 						FTexelToCorners& CurrentTexelCorners = TexelToCornersMap(X, Y);
 						// Normals need to be unit as their dot product will be used in comparisons later
-						CurrentTexelCorners.WorldTangentZ = CurrentTexelCorners.WorldTangentZ.SizeSquared3() > DELTA ? CurrentTexelCorners.WorldTangentZ.UnsafeNormal3() : FVector4(0,0,1);
+						CurrentTexelCorners.WorldTangentZ = CurrentTexelCorners.WorldTangentZ.SizeSquared3() > DELTA ? CurrentTexelCorners.WorldTangentZ.GetUnsafeNormal3() : FVector4(0,0,1);
 					}
 				}
 
