Bumped Swarm cached static meshes to work around corrupt cache

diff --git a/Engine/Source/Programs/UnrealLightmass/Private/Lighting/Collision.cpp b/Engine/Source/Programs/UnrealLightmass/Private/Lighting/Collision.cpp
--- a/Engine/Source/Programs/UnrealLightmass/Private/Lighting/Collision.cpp
+++ b/Engine/Source/Programs/UnrealLightmass/Private/Lighting/Collision.cpp
@@ -76,6 +76,8 @@ void FStaticLightingAggregateMesh::AddMesh(const FStaticLightingMesh* Mesh, cons
 			int32 I2 = 0;
 			Mesh->GetTriangleIndices(TriangleIndex,I0,I1,I2);
 
+			check(I0 <= Mesh->NumVertices && I1 <= Mesh->NumVertices && I2 <= Mesh->NumVertices);
+
 			const bool bTwoSided = Mesh->IsTwoSided(ElementIndex) || Mesh->IsCastingShadowAsTwoSided();
 			const bool bStaticAndOpaque = !Mesh->IsMasked(ElementIndex) && !Mesh->IsTranslucent(ElementIndex) && !Mesh->bMovable;
 
diff --git a/Engine/Source/Programs/UnrealLightmass/Public/ImportExport.h b/Engine/Source/Programs/UnrealLightmass/Public/ImportExport.h
--- a/Engine/Source/Programs/UnrealLightmass/Public/ImportExport.h
+++ b/Engine/Source/Programs/UnrealLightmass/Public/ImportExport.h
@@ -48,7 +48,7 @@ namespace Lightmass
 	static const FGuid LM_MESHAREALIGHTDATA_VERSION		= FGuid(0x7b05d71b, 0x81484a07, 0xa15e44ba, 0x76f4634f);
 	static const FGuid LM_DEBUGOUTPUT_VERSION			= FGuid(0xfa520216, 0xd5aa49eb, 0xa9bbc1f5, 0x8ab715fc);
 	static const FGuid LM_SCENE_VERSION					= FGuid(0x8ad6acb1, 0xe15e4fec, 0x93c72936, 0xf2b67090);
-	static const FGuid LM_STATICMESH_VERSION			= FGuid(0x9b3b1bc6, 0x9270497a, 0xa07521fe, 0xc7045e91);
+	static const FGuid LM_STATICMESH_VERSION			= FGuid(0x9b3b1bc6, 0x9270497a, 0xa07522fe, 0xc7055e91);
 	static const FGuid LM_MATERIAL_VERSION				= FGuid(0xa6640d72, 0x33bf47b4, 0xbbcf76bb, 0x4350a2b5);
 
 
