Fix Landscape lighting when rotated
Normals weren't being rotated

diff --git a/Engine/Source/Programs/UnrealLightmass/Private/Lighting/Landscape.cpp b/Engine/Source/Programs/UnrealLightmass/Private/Lighting/Landscape.cpp
--- a/Engine/Source/Programs/UnrealLightmass/Private/Lighting/Landscape.cpp
+++ b/Engine/Source/Programs/UnrealLightmass/Private/Lighting/Landscape.cpp
@@ -55,25 +55,28 @@ namespace Lightmass
 		int32 X, Y;
 		VertexIndexToXY(VertexIndex, X, Y);
 
-		//GetWorldPositionTangents(X, Y, OutVertex.WorldPosition, OutVertex.WorldTangentX, OutVertex.WorldTangentY, OutVertex.WorldTangentZ);
-		int32 LocalX = X-ExpandQuadsX;
-		int32 LocalY = Y-ExpandQuadsY;
+		const int32 LocalX = X - ExpandQuadsX;
+		const int32 LocalY = Y - ExpandQuadsY;
 
 		const FColor* Data = GetHeightData( X, Y );
 
-		OutVertex.WorldTangentZ.X = 2.f / 255.f * (float)Data->B - 1.f;
-		OutVertex.WorldTangentZ.Y = 2.f / 255.f * (float)Data->A - 1.f;
-		OutVertex.WorldTangentZ.Z = FMath::Sqrt(FMath::Max(1.f - (FMath::Square(OutVertex.WorldTangentZ.X)+FMath::Square(OutVertex.WorldTangentZ.Y)), 0.f));
-		OutVertex.WorldTangentX = FVector4(OutVertex.WorldTangentZ.Z, 0.f, -OutVertex.WorldTangentZ.X);
+		OutVertex.WorldTangentZ.X = 2.0f / 255.f * (float)Data->B - 1.0f;
+		OutVertex.WorldTangentZ.Y = 2.0f / 255.f * (float)Data->A - 1.0f;
+		OutVertex.WorldTangentZ.Z = FMath::Sqrt(FMath::Max(1.0f - (FMath::Square(OutVertex.WorldTangentZ.X) + FMath::Square(OutVertex.WorldTangentZ.Y)), 0.f));
+		OutVertex.WorldTangentX = FVector4(OutVertex.WorldTangentZ.Z, 0.0f, -OutVertex.WorldTangentZ.X);
 		OutVertex.WorldTangentY = OutVertex.WorldTangentZ ^ OutVertex.WorldTangentX;
 
-		// Assume there is no rotation, so we don't need to do any LocalToWorld.
-		uint16 Height = (Data->R << 8) + Data->G;
+		// Copied (vaguely) from FLandscapeComponentDataInterface::GetWorldPositionTangents to fix bad lighting when rotated
+		const FMatrix LtWNoScale = LocalToWorld.GetMatrixWithoutScale();
+		OutVertex.WorldTangentX = LtWNoScale.TransformVector(OutVertex.WorldTangentX);
+		OutVertex.WorldTangentY = LtWNoScale.TransformVector(OutVertex.WorldTangentY);
+		OutVertex.WorldTangentZ = LtWNoScale.TransformVector(OutVertex.WorldTangentZ);
 
+		const uint16 Height = (Data->R << 8) + Data->G;
 		OutVertex.WorldPosition = LocalToWorld.TransformPosition( FVector4( LocalX, LocalY, ((float)Height - 32768.f) * LANDSCAPE_ZSCALE ) );
-		//UE_LOG(LogLightmass, Log, TEXT("%d, %d, %d, %d, %d, %d, X:%f, Y:%f, Z:%f "), SectionBaseX+LocalX-ExpandQuadsX, SectionBaseY+LocalY-ExpandQuadsY, ClampedLocalX, ClampedLocalY, SectionBaseX, SectionBaseY, WorldPos.X, WorldPos.Y, WorldPos.Z);
+		//UE_LOG(LogLightmass, Log, TEXT("%d, %d, %d, %d, %d, %d, X:%f, Y:%f, Z:%f "), SectionBaseX + LocalX - ExpandQuadsX, SectionBaseY + LocalY - ExpandQuadsY, ClampedLocalX, ClampedLocalY, SectionBaseX, SectionBaseY, WorldPos.X, WorldPos.Y, WorldPos.Z);
 
-		int32 LightmapUVIndex = 1;
+		const int32 LightmapUVIndex = 1;
 
 		OutVertex.TextureCoordinates[0] = FVector2D((float)X / NumVertices, (float)Y / NumVertices); 
 		OutVertex.TextureCoordinates[LightmapUVIndex].X = X * UVFactor;
diff --git a/Engine/Source/Runtime/Landscape/Private/LandscapeLight.cpp b/Engine/Source/Runtime/Landscape/Private/LandscapeLight.cpp
--- a/Engine/Source/Runtime/Landscape/Private/LandscapeLight.cpp
+++ b/Engine/Source/Runtime/Landscape/Private/LandscapeLight.cpp
@@ -516,21 +516,24 @@ void FLandscapeStaticLightingMesh::GetStaticLightingVertex(int32 VertexIndex, FS
 	const int32 X = VertexIndex % NumVertices;
 	const int32 Y = VertexIndex / NumVertices;
 
-	int32 LocalX = X-ExpandQuadsX;
-	int32 LocalY = Y-ExpandQuadsY;
+	const int32 LocalX = X - ExpandQuadsX;
+	const int32 LocalY = Y - ExpandQuadsY;
 
 	const FColor* Data = &HeightData[X + Y * NumVertices];
 
-	OutVertex.WorldTangentZ.X = 2.f / 255.f * (float)Data->B - 1.f;
-	OutVertex.WorldTangentZ.Y = 2.f / 255.f * (float)Data->A - 1.f;
-	OutVertex.WorldTangentZ.Z = FMath::Sqrt(1.f - (FMath::Square(OutVertex.WorldTangentZ.X)+FMath::Square(OutVertex.WorldTangentZ.Y)));
-	OutVertex.WorldTangentX = FVector4(OutVertex.WorldTangentZ.Z, 0.f, -OutVertex.WorldTangentZ.X);
+	OutVertex.WorldTangentZ.X = 2.0f / 255.f * (float)Data->B - 1.0f;
+	OutVertex.WorldTangentZ.Y = 2.0f / 255.f * (float)Data->A - 1.0f;
+	OutVertex.WorldTangentZ.Z = FMath::Sqrt(1.0f - (FMath::Square(OutVertex.WorldTangentZ.X) + FMath::Square(OutVertex.WorldTangentZ.Y)));
+	OutVertex.WorldTangentX = FVector4(OutVertex.WorldTangentZ.Z, 0.0f, -OutVertex.WorldTangentZ.X);
 	OutVertex.WorldTangentY = OutVertex.WorldTangentZ ^ OutVertex.WorldTangentX;
 
-	// Assume there is no rotation, so we don't need to do any LocalToWorld.
-	uint16 Height = (Data->R << 8) + Data->G;
+	// Copied from FLandscapeComponentDataInterface::GetWorldPositionTangents to fix bad lighting when rotated
+	OutVertex.WorldTangentX = LocalToWorld.TransformVectorNoScale(OutVertex.WorldTangentX);
+	OutVertex.WorldTangentY = LocalToWorld.TransformVectorNoScale(OutVertex.WorldTangentY);
+	OutVertex.WorldTangentZ = LocalToWorld.TransformVectorNoScale(OutVertex.WorldTangentZ);
 
-	OutVertex.WorldPosition = LocalToWorld.TransformPosition( FVector( LocalX, LocalY, LandscapeDataAccess::GetLocalHeight(Height) ) );	
+	const uint16 Height = (Data->R << 8) + Data->G;
+	OutVertex.WorldPosition = LocalToWorld.TransformPosition(FVector(LocalX, LocalY, LandscapeDataAccess::GetLocalHeight(Height)));
 
 	OutVertex.TextureCoordinates[0] = FVector2D((float)X / NumVertices, (float)Y / NumVertices); 
 	OutVertex.TextureCoordinates[LANDSCAPE_LIGHTMAP_UV_INDEX].X = X * UVFactor;
