UnrealLightmass: fixed shadowing bugs.

diff --git a/Engine/Source/Programs/UnrealLightmass/Private/Lighting/LightingMesh.cpp b/Engine/Source/Programs/UnrealLightmass/Private/Lighting/LightingMesh.cpp
--- a/Engine/Source/Programs/UnrealLightmass/Private/Lighting/LightingMesh.cpp
+++ b/Engine/Source/Programs/UnrealLightmass/Private/Lighting/LightingMesh.cpp
@@ -111,18 +111,18 @@ void FStaticLightingMesh::Import( FLightmassImporter& Importer )
 	for (int32 MtrlIdx = 0; MtrlIdx < NumMaterialElements; MtrlIdx++)
 	{
 		FMaterialElement& CurrentMaterialElement = MaterialElements[MtrlIdx];
-		FMaterialElementData TempData;
-		Importer.ImportData(&TempData);
-		CurrentMaterialElement.MaterialId = TempData.MaterialId;
-		CurrentMaterialElement.bUseTwoSidedLighting = TempData.bUseTwoSidedLighting;
-		CurrentMaterialElement.bShadowIndirectOnly = TempData.bShadowIndirectOnly;
-		CurrentMaterialElement.bUseEmissiveForStaticLighting = TempData.bUseEmissiveForStaticLighting;
+		FMaterialElementData METempData;
+		Importer.ImportData(&METempData);
+		CurrentMaterialElement.MaterialId = METempData.MaterialId;
+		CurrentMaterialElement.bUseTwoSidedLighting = METempData.bUseTwoSidedLighting;
+		CurrentMaterialElement.bShadowIndirectOnly = METempData.bShadowIndirectOnly;
+		CurrentMaterialElement.bUseEmissiveForStaticLighting = METempData.bUseEmissiveForStaticLighting;
 		// Validating data here instead of in Unreal since EmissiveLightFalloffExponent is used in so many different object types
-		CurrentMaterialElement.EmissiveLightFalloffExponent = FMath::Max(TempData.EmissiveLightFalloffExponent, 0.0f);
-		CurrentMaterialElement.EmissiveLightExplicitInfluenceRadius = FMath::Max(TempData.EmissiveLightExplicitInfluenceRadius, 0.0f);
-		CurrentMaterialElement.EmissiveBoost = TempData.EmissiveBoost;
-		CurrentMaterialElement.DiffuseBoost = TempData.DiffuseBoost;
-		CurrentMaterialElement.FullyOccludedSamplesFraction = TempData.FullyOccludedSamplesFraction;
+		CurrentMaterialElement.EmissiveLightFalloffExponent = FMath::Max(METempData.EmissiveLightFalloffExponent, 0.0f);
+		CurrentMaterialElement.EmissiveLightExplicitInfluenceRadius = FMath::Max(METempData.EmissiveLightExplicitInfluenceRadius, 0.0f);
+		CurrentMaterialElement.EmissiveBoost = METempData.EmissiveBoost;
+		CurrentMaterialElement.DiffuseBoost = METempData.DiffuseBoost;
+		CurrentMaterialElement.FullyOccludedSamplesFraction = METempData.FullyOccludedSamplesFraction;
 		CurrentMaterialElement.Material = Importer.ConditionalImportObject<FMaterial>(CurrentMaterialElement.MaterialId, LM_MATERIAL_VERSION, LM_MATERIAL_EXTENSION, LM_MATERIAL_CHANNEL_FLAGS, Importer.GetMaterials());
 		checkf(CurrentMaterialElement.Material, TEXT("Failed to import material with GUID %s"), *CurrentMaterialElement.MaterialId.ToString());
 
diff --git a/Engine/Source/Programs/UnrealLightmass/Private/Lighting/LightingSystem.h b/Engine/Source/Programs/UnrealLightmass/Private/Lighting/LightingSystem.h
--- a/Engine/Source/Programs/UnrealLightmass/Private/Lighting/LightingSystem.h
+++ b/Engine/Source/Programs/UnrealLightmass/Private/Lighting/LightingSystem.h
@@ -2612,7 +2612,7 @@ private:
 	 * Each array was generated on a separate thread, so these are stored as an array of irradiance photon arrays,
 	 * Which avoids copying to one large array, since that can take a while due to the large irradiance photon memory size.
 	 */
-	TArray<TArray<FIrradiancePhoton>> IrradiancePhotons;
+	TArray<TArray<FIrradiancePhoton>> mIrradiancePhotons;
 
 	/** Maximum distance to trace a ray through the scene. */
 	float MaxRayDistance;
diff --git a/Engine/Source/Programs/UnrealLightmass/Private/Lighting/PhotonMapping.cpp b/Engine/Source/Programs/UnrealLightmass/Private/Lighting/PhotonMapping.cpp
--- a/Engine/Source/Programs/UnrealLightmass/Private/Lighting/PhotonMapping.cpp
+++ b/Engine/Source/Programs/UnrealLightmass/Private/Lighting/PhotonMapping.cpp
@@ -130,12 +130,12 @@ void FStaticLightingSystem::EmitPhotons()
 	}
 	
 	// Presize for the results from two emitting passes (direct photon emitting, then indirect)
-	IrradiancePhotons.Empty(NumPhotonWorkRanges * 2);
+	mIrradiancePhotons.Empty(NumPhotonWorkRanges * 2);
 
 	const double StartEmitDirectTime = FPlatformTime::Seconds();
 	TArray<TArray<FIndirectPathRay> > IndirectPathRays;
 	// Emit photons for the direct photon map, and gather rays which resulted in indirect photon paths.
-	EmitDirectPhotons(ImportanceVolumeBounds, IndirectPathRays, IrradiancePhotons);
+	EmitDirectPhotons(ImportanceVolumeBounds, IndirectPathRays, mIrradiancePhotons);
 
 	const double EndEmitDirectTime = FPlatformTime::Seconds();
 	Stats.EmitDirectPhotonsTime = EndEmitDirectTime - StartEmitDirectTime;
@@ -152,7 +152,7 @@ void FStaticLightingSystem::EmitPhotons()
 	Stats.CachingIndirectPhotonPathsTime = EndCachingIndirectPathsTime - EndEmitDirectTime;
 
 	// Emit photons for the indirect photon map, using the indirect photon paths to guide photon emission.
-	EmitIndirectPhotons(ImportanceVolumeBounds, IndirectPathRays, IrradiancePhotons);
+	EmitIndirectPhotons(ImportanceVolumeBounds, IndirectPathRays, mIrradiancePhotons);
 	const double EndEmitIndirectTime = FPlatformTime::Seconds();
 	Stats.EmitIndirectPhotonsTime = EndEmitIndirectTime - EndCachingIndirectPathsTime;
 	LogSolverMessage(FString::Printf(TEXT("EmitIndirectPhotons complete, %.3f million photons emitted in %.1f seconds"), Stats.NumSecondPassPhotonsEmitted / 1000000.0f, Stats.EmitIndirectPhotonsTime));
@@ -162,7 +162,7 @@ void FStaticLightingSystem::EmitPhotons()
 		// Process all irradiance photons and mark ones that have direct photons nearby,
 		// So that we can search for those with a smaller radius when using them for rendering.
 		// This allows more accurate direct shadow transitions with irradiance photons.
-		MarkIrradiancePhotons(ImportanceVolumeBounds, IrradiancePhotons);
+		MarkIrradiancePhotons(ImportanceVolumeBounds, mIrradiancePhotons);
 		const double EndMarkIrradiancePhotonsTime = FPlatformTime::Seconds();
 		Stats.IrradiancePhotonMarkingTime = EndMarkIrradiancePhotonsTime - EndEmitIndirectTime;
 		LogSolverMessage(FString::Printf(TEXT("Marking Irradiance Photons complete, %.3f million photons marked in %.1f seconds"), Stats.NumIrradiancePhotons / 1000000.0f, Stats.IrradiancePhotonMarkingTime));
@@ -182,7 +182,7 @@ void FStaticLightingSystem::EmitPhotons()
 		// The optimization is that irradiance is pre-calculated at a subset of the photons so that final gather rays can just lookup the nearest irradiance photon,
 		// Instead of doing a photon density estimation to calculate irradiance.
 		const double StartCalculateIrradiancePhotonsTime = FPlatformTime::Seconds();
-		CalculateIrradiancePhotons(ImportanceVolumeBounds, IrradiancePhotons);
+		CalculateIrradiancePhotons(ImportanceVolumeBounds, mIrradiancePhotons);
 		Stats.IrradiancePhotonCalculatingTime = FPlatformTime::Seconds() - StartCalculateIrradiancePhotonsTime;
 		LogSolverMessage(FString::Printf(TEXT("Calculate Irradiance Photons complete, %.3f million irradiance calculations in %.1f seconds"), Stats.NumFoundIrradiancePhotons / 1000000.0f, Stats.IrradiancePhotonCalculatingTime));
 
@@ -1362,7 +1362,7 @@ uint32 FIrradiancePhotonCalculatingThreadRunnable::Run()
 void FStaticLightingSystem::CalculateIrradiancePhotonsThreadLoop(
 	int32 ThreadIndex, 
 	TArray<TArray<FIrradiancePhoton>>& IrradiancePhotons, 
-	FCalculateIrradiancePhotonStats& Stats)
+	FCalculateIrradiancePhotonStats& OutStats)
 {
 	while (true)
 	{
@@ -1372,7 +1372,7 @@ void FStaticLightingSystem::CalculateIrradiancePhotonsThreadLoop(
 		const int32 RangeIndex = IrradianceCalcWorkRangeIndex.Increment() - 1;
 		if (RangeIndex < IrradianceCalculationWorkRanges.Num())
 		{
-			CalculateIrradiancePhotonsWorkRange(IrradiancePhotons, IrradianceCalculationWorkRanges[RangeIndex], Stats);
+			CalculateIrradiancePhotonsWorkRange(IrradiancePhotons, IrradianceCalculationWorkRanges[RangeIndex], OutStats);
 		}
 		else
 		{
@@ -1386,7 +1386,7 @@ void FStaticLightingSystem::CalculateIrradiancePhotonsThreadLoop(
 void FStaticLightingSystem::CalculateIrradiancePhotonsWorkRange(
 	TArray<TArray<FIrradiancePhoton>>& IrradiancePhotons, 
 	FIrradianceCalculatingWorkRange WorkRange,
-	FCalculateIrradiancePhotonStats& Stats)
+	FCalculateIrradiancePhotonStats& OutStats)
 {
 	// Temporary array that is reused for all photon searches by this thread, to reduce allocations
 	TArray<FPhoton> TempFoundPhotons;
@@ -1419,7 +1419,7 @@ void FStaticLightingSystem::CalculateIrradiancePhotonsWorkRange(
 				CurrentIrradiancePhoton,
 				bDebugThisPhoton && GeneralSettings.ViewSingleBounceNumber == 0,
 				TempFoundPhotons,
-				Stats);
+				OutStats);
 
 			checkSlow(FLinearColorUtils::AreFloatsValid(DirectPhotonIrradiance));
 
@@ -1445,7 +1445,7 @@ void FStaticLightingSystem::CalculateIrradiancePhotonsWorkRange(
 				CurrentIrradiancePhoton,
 				bDebugThisPhoton && GeneralSettings.ViewSingleBounceNumber == 1,
 				TempFoundPhotons,
-				Stats);
+				OutStats);
 
 			checkSlow(FLinearColorUtils::AreFloatsValid(FirstBouncePhotonIrradiance));
 
@@ -1469,7 +1469,7 @@ void FStaticLightingSystem::CalculateIrradiancePhotonsWorkRange(
 					CurrentIrradiancePhoton,
 					bDebugThisPhoton && GeneralSettings.ViewSingleBounceNumber > 1,
 					TempFoundPhotons,
-					Stats);
+					OutStats);
 
 				checkSlow(FLinearColorUtils::AreFloatsValid(SecondBouncePhotonIrradiance));
 
@@ -1532,9 +1532,9 @@ void FStaticLightingSystem::CacheIrradiancePhotonsThreadLoop(int32 ThreadIndex,
 			if (bIsMainThread)
 			{
 				// Check the health of all static lighting threads.
-				for (int32 ThreadIndex = 0; ThreadIndex < IrradiancePhotonCachingThreads.Num(); ThreadIndex++)
+				for (int32 ThreadIndexIter = 0; ThreadIndexIter < IrradiancePhotonCachingThreads.Num(); ThreadIndexIter++)
 				{
-					IrradiancePhotonCachingThreads[ThreadIndex].CheckHealth();
+					IrradiancePhotonCachingThreads[ThreadIndexIter].CheckHealth();
 				}
 			}
 
@@ -2126,7 +2126,7 @@ FLinearColor FStaticLightingSystem::CalculatePhotonIrradiance(
 	const FIrradiancePhoton& IrradiancePhoton,
 	bool bDebugThisCalculation,
 	TArray<FPhoton>& TempFoundPhotons,
-	FCalculateIrradiancePhotonStats& Stats) const
+	FCalculateIrradiancePhotonStats& OutStats) const
 {
 	// Empty TempFoundPhotons without causing any allocations / frees
 	TempFoundPhotons.Empty(TempFoundPhotons.Num() + TempFoundPhotons.GetSlack());
@@ -2140,13 +2140,13 @@ FLinearColor FStaticLightingSystem::CalculatePhotonIrradiance(
 		false,
 		bDebugThisCalculation, 
 		TempFoundPhotons, 
-		Stats);
+		OutStats);
 
 	FLinearColor PhotonIrradiance(FLinearColor::Black);
 
 	if (TempFoundPhotons.Num() > 0)
 	{
-		LIGHTINGSTAT(FScopedRDTSCTimer CalculateIrradianceTimer(Stats.CalculateIrradianceThreadTime));
+		LIGHTINGSTAT(FScopedRDTSCTimer CalculateIrradianceTimer(OutStats.CalculateIrradianceThreadTime));
 		const float MaxFoundDistance = FMath::Sqrt(MaxFoundDistanceSquared);
 		// Estimate the photon density using a cone filter, from the paper "Global Illumination using Photon Maps"
 		const float DiskArea = (float)PI * MaxFoundDistanceSquared;
diff --git a/Engine/Source/Programs/UnrealLightmass/Private/Lighting/TextureMapping.cpp b/Engine/Source/Programs/UnrealLightmass/Private/Lighting/TextureMapping.cpp
--- a/Engine/Source/Programs/UnrealLightmass/Private/Lighting/TextureMapping.cpp
+++ b/Engine/Source/Programs/UnrealLightmass/Private/Lighting/TextureMapping.cpp
@@ -2758,15 +2758,15 @@ void FStaticLightingSystem::CalculateIndirectLightingTextureMapping(
 		} 
 		while (TextureMapping->NumOutstandingCacheTasks > 0);
 
-		TArray<FCacheIndirectTaskDescription*> CompletedTasks;
-		TextureMapping->CompletedCacheIndirectLightingTasks.PopAll<TArray<FCacheIndirectTaskDescription*>, FCacheIndirectTaskDescription*>(CompletedTasks);
-		check(CompletedTasks.Num() == NumTasksSubmitted);
+		TArray<FCacheIndirectTaskDescription*> CompletedCILTasks;
+		TextureMapping->CompletedCacheIndirectLightingTasks.PopAll<TArray<FCacheIndirectTaskDescription*>, FCacheIndirectTaskDescription*>(CompletedCILTasks);
+		check(CompletedCILTasks.Num() == NumTasksSubmitted);
 
 		int32 NextRecordId = 0;
 
-		for (int32 TaskIndex = 0; TaskIndex < CompletedTasks.Num(); TaskIndex++)
+		for (int32 TaskIndex = 0; TaskIndex < CompletedCILTasks.Num(); TaskIndex++)
 		{
-			FCacheIndirectTaskDescription* Task = CompletedTasks[TaskIndex];
+			FCacheIndirectTaskDescription* Task = CompletedCILTasks[TaskIndex];
 
 			TArray<TLightingCache<FFinalGatherSample>::FRecord<FFinalGatherSample> > Records;
 			Task->MappingContext.FirstBounceCache.GetAllRecords(Records);
@@ -2789,7 +2789,7 @@ void FStaticLightingSystem::CalculateIndirectLightingTextureMapping(
 
 			// Note: the task's MappingContext stats will be merged into the global stats automatically due to the MappingContext destructor
 
-			delete CompletedTasks[TaskIndex];
+			delete CompletedCILTasks[TaskIndex];
 		}
 
 		const double EndCacheTime = FPlatformTime::Seconds();
@@ -2804,7 +2804,7 @@ void FStaticLightingSystem::CalculateIndirectLightingTextureMapping(
 			}
 
 			const int32 InterpolationTaskSize = IrradianceCachingSettings.InterpolateTaskSize;
-			int32 NumTasksSubmitted = 0;
+			int32 NumIILTasksSubmitted = 0;
 
 			// Break this mapping into multiple interpolation tasks in texture space blocks
 			for (int32 TaskY = 0; TaskY < TextureMapping->CachedSizeY; TaskY += InterpolationTaskSize)
@@ -2827,7 +2827,7 @@ void FStaticLightingSystem::CalculateIndirectLightingTextureMapping(
 							|| Scene.DebugInput.LocalX >= TaskX && Scene.DebugInput.LocalX < TaskX + InterpolationTaskSize
 							&& Scene.DebugInput.LocalY >= TaskY && Scene.DebugInput.LocalY < TaskY + InterpolationTaskSize);
 
-					NumTasksSubmitted++;
+					NumIILTasksSubmitted++;
 					FPlatformAtomics::InterlockedIncrement(&TextureMapping->NumOutstandingInterpolationTasks);
 					InterpolateIndirectLightingTasks.Push(NewTask);
 				}
@@ -2853,7 +2853,7 @@ void FStaticLightingSystem::CalculateIndirectLightingTextureMapping(
 
 			TArray<FInterpolateIndirectTaskDescription*> CompletedTasks;
 			TextureMapping->CompletedInterpolationTasks.PopAll<TArray<FInterpolateIndirectTaskDescription*>, FInterpolateIndirectTaskDescription*>(CompletedTasks);
-			check(CompletedTasks.Num() == NumTasksSubmitted);
+			check(CompletedTasks.Num() == NumIILTasksSubmitted);
 
 			for (int32 TaskIndex = 0; TaskIndex < CompletedTasks.Num(); TaskIndex++)
 			{
diff --git a/Engine/Source/Programs/UnrealLightmass/Private/LightmassCore/Misc/LMThreading.cpp b/Engine/Source/Programs/UnrealLightmass/Private/LightmassCore/Misc/LMThreading.cpp
--- a/Engine/Source/Programs/UnrealLightmass/Private/LightmassCore/Misc/LMThreading.cpp
+++ b/Engine/Source/Programs/UnrealLightmass/Private/LightmassCore/Misc/LMThreading.cpp
@@ -71,18 +71,18 @@ FSemaphore::~FSemaphore(void)
  *
  * @return Returns true if the event was created, false otherwise
  */
-bool FSemaphore::Create(int32 MaxCount)
+bool FSemaphore::Create(int32 InMaxCount)
 {
 	// Create the event and default it to non-signaled
 #if PLATFORM_WINDOWS
-	Semaphore = CreateSemaphore(NULL, 0, MaxCount, NULL);
+	Semaphore = CreateSemaphore(NULL, 0, InMaxCount, NULL);
 	return Semaphore != NULL;
 #else
 	if( !bInitialized )
 	{
-		check( MaxCount > 0 );
-		this->MaxCount = MaxCount;
-		this->Counter = 0;
+		check( InMaxCount > 0 );
+		MaxCount = InMaxCount;
+		Counter = 0;
 
 		int Result = pthread_mutex_init( &Guard, NULL );
 		if( Result == 0 )
