Fixed Lightmass crash from new FString[] semantics

diff --git a/Engine/Source/Programs/UnrealLightmass/Private/LightmassCore/LMCore.cpp b/Engine/Source/Programs/UnrealLightmass/Private/LightmassCore/LMCore.cpp
--- a/Engine/Source/Programs/UnrealLightmass/Private/LightmassCore/LMCore.cpp
+++ b/Engine/Source/Programs/UnrealLightmass/Private/LightmassCore/LMCore.cpp
@@ -22,7 +22,7 @@ FLightmassLog::FLightmassLog()
 	ExeName = ExeName.Replace(TEXT("\\"), TEXT("/"));
 	// Extract filename part and add "-[guid].log"
 	int32 ExeNameLen = ExeName.Len();
-	int32 PathSeparatorPos = ExeNameLen;
+	int32 PathSeparatorPos = ExeNameLen - 1;
 	while ( PathSeparatorPos >= 0 && ExeName[PathSeparatorPos] != TEXT('/') )
 	{
 		PathSeparatorPos--;
