Area shadowing for precomputed shadows from stationary lights
* New light Lightmass setting 'bUseAreaShadowsForStationaryLight'
* Removed 'DistanceFieldPenumbraSize' project setting, uniform penumbra size comes from each light's LightSourceAngle or SourceRadius using a heuristic now

diff --git a/Engine/Shaders/LightmapCommon.usf b/Engine/Shaders/LightmapCommon.usf
--- a/Engine/Shaders/LightmapCommon.usf
+++ b/Engine/Shaders/LightmapCommon.usf
@@ -107,7 +107,7 @@ float4 GetSkyBentNormalAndOcclusion(float2 LightmapUV)
 
 #if STATICLIGHTING_TEXTUREMASK && STATICLIGHTING_SIGNEDDISTANCEFIELD
 	float4 StaticShadowMapMasks;
-	float4 DistanceFieldParameters;
+	float4 InvUniformPenumbraSizes;
 	Texture2D StaticShadowTexture;
 	SamplerState StaticShadowTextureSampler;
 #endif
@@ -118,8 +118,10 @@ half4 GetPrecomputedShadowMasks(FVertexFactoryInterpolantsVSToPS Interpolants)
 
 		// Fetch the 4 channels of distance field data
 		half4 DistanceField = Texture2DSample(StaticShadowTexture, StaticShadowTextureSampler, GetShadowMapCoordinate(Interpolants));
+		float4 DistanceFieldBias = -.5f * InvUniformPenumbraSizes + .5f;
+
 		// Compute shadow factors by scaling and biasing the distance
-		half4 ShadowFactors = saturate( DistanceField * DistanceFieldParameters.x + DistanceFieldParameters.y );
+		half4 ShadowFactors = saturate(DistanceField * InvUniformPenumbraSizes + DistanceFieldBias);
 		return StaticShadowMapMasks * ShadowFactors * ShadowFactors;
 
 	#elif HQ_TEXTURE_LIGHTMAP || LQ_TEXTURE_LIGHTMAP
@@ -142,8 +144,9 @@ half GetPrimaryPrecomputedShadowMask(FVertexFactoryInterpolantsVSToPS Interpolan
 
 		// Fetch the distance field data
 		half DistanceField = Texture2DSample(StaticShadowTexture, StaticShadowTextureSampler, GetShadowMapCoordinate(Interpolants)).r;
+		float DistanceFieldBias = -.5f * InvUniformPenumbraSizes.x + .5f;
 		// Compute shadow factors by scaling and biasing the distance
-		half ShadowFactor = saturate( DistanceField * DistanceFieldParameters.x + DistanceFieldParameters.y );
+		half ShadowFactor = saturate( DistanceField * InvUniformPenumbraSizes.x + DistanceFieldBias );
 		return ShadowFactor * ShadowFactor;
 
 	#elif MOVABLE_DIRECTIONAL_LIGHT
diff --git a/Engine/Source/Editor/UnrealEd/Private/LevelViewportClickHandlers.cpp b/Engine/Source/Editor/UnrealEd/Private/LevelViewportClickHandlers.cpp
--- a/Engine/Source/Editor/UnrealEd/Private/LevelViewportClickHandlers.cpp
+++ b/Engine/Source/Editor/UnrealEd/Private/LevelViewportClickHandlers.cpp
@@ -156,7 +156,6 @@ namespace ClickHandlers
 		}
 		else if( Click.GetKey() != EKeys::RightMouseButton )
 		{
-#if ALLOW_LIGHTMAP_SAMPLE_DEBUGGING
 			if ( Click.GetKey() == EKeys::LeftMouseButton && ViewportClient->Viewport->KeyState(EKeys::T) && Actor )
 			{
 				TArray<UActorComponent*> Components;
@@ -164,7 +163,6 @@ namespace ClickHandlers
 				SetDebugLightmapSample(&Components, NULL, 0, GEditor->ClickLocation);
 			}
 			else 
-#endif
 			if( Click.GetKey() == EKeys::LeftMouseButton && ViewportClient->Viewport->KeyState(EKeys::L) )
 			{
 				// If shift is down, we pick a color from under the mouse in the viewport and create a light with that color.
diff --git a/Engine/Source/Editor/UnrealEd/Private/Lightmass/Lightmass.cpp b/Engine/Source/Editor/UnrealEd/Private/Lightmass/Lightmass.cpp
--- a/Engine/Source/Editor/UnrealEd/Private/Lightmass/Lightmass.cpp
+++ b/Engine/Source/Editor/UnrealEd/Private/Lightmass/Lightmass.cpp
@@ -124,11 +124,10 @@ void Copy( const ULightComponentBase* In, Lightmass::FLightData& Out )
 	{
 		Out.LightFlags |= Lightmass::GI_LIGHT_HASSTATICSHADOWING;
 		Out.LightFlags |= Lightmass::GI_LIGHT_HASSTATICLIGHTING;
-		Out.LightFlags |= Lightmass::GI_LIGHT_USEDIRECTLIGHTMAP;
 	}
 	else if (In->HasStaticShadowing())
 	{
-		Out.LightFlags |= Lightmass::GI_LIGHT_USESIGNEDDISTANCEFIELDSHADOWS;
+		Out.LightFlags |= Lightmass::GI_LIGHT_STORE_SEPARATE_SHADOW_FACTOR;
 		Out.LightFlags |= Lightmass::GI_LIGHT_HASSTATICSHADOWING;
 	}
 
@@ -155,6 +154,11 @@ void Copy( const ULightComponent* In, Lightmass::FLightData& Out )
 		Out.LightFlags |= Lightmass::GI_LIGHT_INVERSE_SQUARED;
 	}
 
+	if (In->GetLightmassSettings().bUseAreaShadowsForStationaryLight)
+	{
+		Out.LightFlags |= Lightmass::GI_LIGHT_USE_AREA_SHADOWS_FOR_SEPARATE_SHADOW_FACTOR;
+	}
+
 	Out.Brightness = In->ComputeLightBrightness();
 	Out.Position = In->GetLightPosition();
 	Out.Direction = In->GetDirection();
@@ -1350,7 +1354,6 @@ void FLightmassExporter::WriteBaseMappingData( int32 Channel, const FStaticLight
 	FMemory::Memzero(&MappingData,sizeof(MappingData));
 	MappingData.Guid = Mapping->Mesh->Guid;
 	MappingData.StaticLightingMeshInstance = Mapping->Mesh->SourceMeshGuid;
-	MappingData.bForceDirectLightMap = false;
 	Swarm.WriteChannel( Channel, &MappingData, sizeof(MappingData) );
 }
 
diff --git a/Engine/Source/Editor/UnrealEd/Private/StaticLightingSystem/StaticLightingDebug.cpp b/Engine/Source/Editor/UnrealEd/Private/StaticLightingSystem/StaticLightingDebug.cpp
--- a/Engine/Source/Editor/UnrealEd/Private/StaticLightingSystem/StaticLightingDebug.cpp
+++ b/Engine/Source/Editor/UnrealEd/Private/StaticLightingSystem/StaticLightingDebug.cpp
@@ -162,9 +162,12 @@ static bool GetBarycentricWeights(
 	return false;
 }
 
+float TriangleTolerance = .01f;
+
 /** Updates GCurrentSelectedLightmapSample given a selected actor's components and the location of the click. */
 void SetDebugLightmapSample(TArray<UActorComponent*>* Components, UModel* Model, int32 iSurf, FVector ClickLocation)
 {
+#if ALLOW_LIGHTMAP_SAMPLE_DEBUGGING
 	UStaticMeshComponent* SMComponent = NULL;
 	if (Components)
 	{
@@ -225,7 +228,7 @@ void SetDebugLightmapSample(TArray<UActorComponent*>* Components, UModel* Model,
 
 				FVector BaryCentricWeights;
 				// Continue if click location is in the triangle and get its barycentric weights
-				if (GetBarycentricWeights(Position0, Position1, Position2, ClickLocation, .001f, BaryCentricWeights))
+				if (GetBarycentricWeights(Position0, Position1, Position2, ClickLocation, TriangleTolerance, BaryCentricWeights))
 				{
 					RestoreSelectedLightmapSample();
 
@@ -439,6 +442,7 @@ void SetDebugLightmapSample(TArray<UActorComponent*>* Components, UModel* Model,
 		RestoreSelectedLightmapSample();
 		GCurrentSelectedLightmapSample = FSelectedLightmapSample();
 	}
+#endif
 }
 
 /** Renders debug elements for visualizing static lighting info */
diff --git a/Engine/Source/Programs/UnrealLightmass/Private/ImportExport/LightmassScene.h b/Engine/Source/Programs/UnrealLightmass/Private/ImportExport/LightmassScene.h
--- a/Engine/Source/Programs/UnrealLightmass/Private/ImportExport/LightmassScene.h
+++ b/Engine/Source/Programs/UnrealLightmass/Private/ImportExport/LightmassScene.h
@@ -187,16 +187,13 @@ public:
 	virtual FVector4 GetDirectLightingDirection(const FVector4& Point, const FVector4& PointNormal) const = 0;
 
 	/**
-	 * Returns whether static lighting, aka lightmaps, is being used for primitive/ light
-	 * interaction.
+	 * Returns whether static lighting, aka lightmaps, is being used for primitive/ light interaction.
 	 *
-	 * @param bForceDirectLightMap	Whether primitive is set to force lightmaps
 	 * @return true if lightmaps/ static lighting is being used, false otherwise
 	 */
-	bool UseStaticLighting( bool bForceDirectLightMap ) const
+	bool UseStaticLighting() const
 	{
-		// use the flags that Unreal sent over to determine if we should use static lighting or not
-		return (LightFlags & GI_LIGHT_HASSTATICLIGHTING) && ((LightFlags & GI_LIGHT_USEDIRECTLIGHTMAP) || bForceDirectLightMap || !(LightFlags & GI_LIGHT_HASSTATICSHADOWING));
+		return (LightFlags & GI_LIGHT_HASSTATICLIGHTING) != 0;
 	}
 
 protected:
diff --git a/Engine/Source/Programs/UnrealLightmass/Private/Lighting/LightingMesh.cpp b/Engine/Source/Programs/UnrealLightmass/Private/Lighting/LightingMesh.cpp
--- a/Engine/Source/Programs/UnrealLightmass/Private/Lighting/LightingMesh.cpp
+++ b/Engine/Source/Programs/UnrealLightmass/Private/Lighting/LightingMesh.cpp
@@ -462,7 +462,7 @@ void FStaticLightingMesh::CreateMeshAreaLights(
 					{
 						// Initialize all of the mesh area light's unused properties to 0
 						FMeshAreaLight* NewLight = new FMeshAreaLight(ForceInit);
-						NewLight->LightFlags = GI_LIGHT_HASSTATICLIGHTING | GI_LIGHT_USEDIRECTLIGHTMAP | GI_LIGHT_CASTSHADOWS | GI_LIGHT_CASTSTATICSHADOWS;
+						NewLight->LightFlags = GI_LIGHT_HASSTATICLIGHTING | GI_LIGHT_CASTSHADOWS | GI_LIGHT_CASTSTATICSHADOWS;
 						NewLight->SetPrimitives(
 							TrimmedEmissivePrimitives[LightIndex], 
 							MaterialElements[MaterialIndex].EmissiveLightFalloffExponent, 
diff --git a/Engine/Source/Programs/UnrealLightmass/Private/Lighting/LightingSystem.cpp b/Engine/Source/Programs/UnrealLightmass/Private/Lighting/LightingSystem.cpp
--- a/Engine/Source/Programs/UnrealLightmass/Private/Lighting/LightingSystem.cpp
+++ b/Engine/Source/Programs/UnrealLightmass/Private/Lighting/LightingSystem.cpp
@@ -1964,7 +1964,7 @@ void FStaticLightingSystem::CalculateApproximateDirectLighting(
 				// Compute the light-map sample for the front-face of the vertex.
 				FGatheredLightSample Lighting = FGatheredLightSample::PointLightWorldSpace(FinalIntensity, TangentLightVector, WorldLightVector.GetSafeNormal());
 
-				if (Light->UseStaticLighting(false) || bCompositeAllLights)
+				if (Light->UseStaticLighting() || bCompositeAllLights)
 				{
 					OutStaticDirectLighting = OutStaticDirectLighting + Lighting;
 				}
diff --git a/Engine/Source/Programs/UnrealLightmass/Private/Lighting/LightingSystem.h b/Engine/Source/Programs/UnrealLightmass/Private/Lighting/LightingSystem.h
--- a/Engine/Source/Programs/UnrealLightmass/Private/Lighting/LightingSystem.h
+++ b/Engine/Source/Programs/UnrealLightmass/Private/Lighting/LightingSystem.h
@@ -2405,17 +2405,28 @@ private:
 		FStaticLightingTextureMapping* TextureMapping, 
 		FStaticLightingMappingContext& MappingContext,
 		FGatheredLightMapData2D& LightMapData, 
-		TMap<const FLight*, FShadowMapData2D*>& ShadowMaps,
+		FShadowMapData2D*& ShadowMapData,
 		const FTexelToVertexMap& TexelToVertexMap, 
 		bool bDebugThisMapping,
 		const FLight* Light,
-		const bool bLowQualityLightMapsOnly ) const;
+		const bool bLowQualityLightMapsOnly) const;
 
 	/** 
 	 * Calculate signed distance field shadowing from a single light,  
 	 * Based on the paper "Improved Alpha-Tested Magnification for Vector Textures and Special Effects" by Valve.
 	 */
-	void CalculateDirectSignedDistanceFieldLightingTextureMapping(
+	void CalculateDirectSignedDistanceFieldLightingTextureMappingTextureSpace(
+		FStaticLightingTextureMapping* TextureMapping, 
+		FStaticLightingMappingContext& MappingContext,
+		FGatheredLightMapData2D& LightMapData, 
+		TMap<const FLight*, FSignedDistanceFieldShadowMapData2D*>& ShadowMaps,
+		const FTexelToVertexMap& TexelToVertexMap, 
+		const FTexelToCornersMap& TexelToCornersMap,
+		bool bDebugThisMapping,
+		const FLight* Light) const;
+
+	/** Experimental method that avoids artifacts due to lightmap seams. */
+	void CalculateDirectSignedDistanceFieldLightingTextureMappingLightSpace(
 		FStaticLightingTextureMapping* TextureMapping, 
 		FStaticLightingMappingContext& MappingContext,
 		FGatheredLightMapData2D& LightMapData, 
diff --git a/Engine/Source/Programs/UnrealLightmass/Private/Lighting/TextureMapping.cpp b/Engine/Source/Programs/UnrealLightmass/Private/Lighting/TextureMapping.cpp
--- a/Engine/Source/Programs/UnrealLightmass/Private/Lighting/TextureMapping.cpp
+++ b/Engine/Source/Programs/UnrealLightmass/Private/Lighting/TextureMapping.cpp
@@ -25,6 +25,9 @@ public:
 		/** Weight used when combining super sampled attributes and determining if the texel has been mapped. */
 		float TotalSampleWeight;
 
+		/** Tracks the max sample weight encountered. */
+		float MaxSampleWeight;
+
 		/** World space radius of the texel. */
 		float TexelRadius;
 
@@ -213,25 +216,32 @@ void FStaticLightingRasterPolicy::ProcessPixel(int32 X,int32 Y,const Interpolant
 	}
 #endif
 
-	if (bUseMaxWeight && SampleWeight > TexelToVertex.TotalSampleWeight)
+	if (bUseMaxWeight)
 	{
-		// Use the sample with the largest weight.  
-		// This has the disadvantage compared averaging based on weight that it won't be well centered for texels on a UV seam,
-		// And a texel spanning multiple triangles will only use the normal from one of those triangles,
-		// But it has the advantage that the final position is guaranteed to be valid (ie actually on a triangle),
-		// Even for split texels which are mapped to triangles in different parts of the mesh.
-		TexelToVertex.TotalSampleWeight = SampleWeight;
-		TexelToVertex.WorldPosition = Interpolant.Vertex.WorldPosition;
-		TexelToVertex.WorldTangentX = Interpolant.Vertex.WorldTangentX;
-		TexelToVertex.WorldTangentY = Interpolant.Vertex.WorldTangentY;
-		TexelToVertex.WorldTangentZ = Interpolant.Vertex.WorldTangentZ;
-		TexelToVertex.TriangleNormal = TriangleNormal;
-		TexelToVertex.ElementIndex = Interpolant.ElementIndex;
-
-		for( int32 CurCoordIndex = 0; CurCoordIndex < MAX_TEXCOORDS; ++CurCoordIndex )
+		if (SampleWeight > TexelToVertex.MaxSampleWeight)
 		{
-			TexelToVertex.TextureCoordinates[ CurCoordIndex ] = Interpolant.Vertex.TextureCoordinates[ CurCoordIndex ];
+			// Use the sample with the largest weight.  
+			// This has the disadvantage compared averaging based on weight that it won't be well centered for texels on a UV seam,
+			// But it has the advantage that the final position is guaranteed to be valid (ie actually on a triangle),
+			// Even for split texels which are mapped to triangles in different parts of the mesh.
+			TexelToVertex.MaxSampleWeight = SampleWeight;
+			TexelToVertex.WorldPosition = Interpolant.Vertex.WorldPosition;
+			TexelToVertex.TriangleNormal = TriangleNormal;
+			TexelToVertex.ElementIndex = Interpolant.ElementIndex;
+
+			for( int32 CurCoordIndex = 0; CurCoordIndex < MAX_TEXCOORDS; ++CurCoordIndex )
+			{
+				TexelToVertex.TextureCoordinates[ CurCoordIndex ] = Interpolant.Vertex.TextureCoordinates[ CurCoordIndex ];
+			}
 		}
+		
+		// Weighted average of normal, improves the case where the position chosen by the max weight has a different normal than the rest of the texel
+		// Eg, small extrusions from an otherwise flat surface, and the texel center lies on the perpendicular extrusion
+		//@todo - only average normals within the texel radius to improve the split texel case?
+		TexelToVertex.WorldTangentX += Interpolant.Vertex.WorldTangentX * SampleWeight;
+		TexelToVertex.WorldTangentY += Interpolant.Vertex.WorldTangentY * SampleWeight;
+		TexelToVertex.WorldTangentZ += Interpolant.Vertex.WorldTangentZ * SampleWeight;
+		TexelToVertex.TotalSampleWeight += SampleWeight;
 	}
 	else if (!bUseMaxWeight)
 	{
@@ -485,10 +495,10 @@ void FStaticLightingSystem::ProcessTextureMapping(FStaticLightingTextureMapping*
 						FDebugStaticLightingVertex DebugVertex;
 						DebugVertex.VertexNormal = FVector4(TexelToVertex.WorldTangentZ);
 						DebugVertex.VertexPosition = TexelToVertex.WorldPosition;
-						
+						DebugOutput.Vertices.Add(DebugVertex);
+
 						if (X == Scene.DebugInput.LocalX && Y == Scene.DebugInput.LocalY)
-						{
-							DebugOutput.Vertices.Add(DebugVertex);
+						{							
 							DebugOutput.SelectedVertexIndices.Add(DebugOutput.Vertices.Num() - 1);
 							DebugOutput.SampleRadius = TexelToVertex.TexelRadius;
 						}
@@ -531,27 +541,77 @@ void FStaticLightingSystem::ProcessTextureMapping(FStaticLightingTextureMapping*
 				}
 				else 
 				{
-					if (!Light->UseStaticLighting(TextureMapping->bForceDirectLightMap)
+					if (!Light->UseStaticLighting()
 						&& (Light->LightFlags & GI_LIGHT_CASTSHADOWS) 
 						&& (Light->LightFlags & GI_LIGHT_CASTSTATICSHADOWS)
-						&& (Light->LightFlags & GI_LIGHT_USESIGNEDDISTANCEFIELDSHADOWS)
+						&& (Light->LightFlags & GI_LIGHT_STORE_SEPARATE_SHADOW_FACTOR)
 						&& ShadowSettings.bAllowSignedDistanceFieldShadows)
 					{
-						// Calculate distance field shadows, where the distance to the nearest shadow transition is stored instead of just a [0,1] shadow factor.
-						CalculateDirectSignedDistanceFieldLightingTextureMapping(TextureMapping, MappingContext, LightMapData, SignedDistanceFieldShadowMaps, TexelToVertexMap, TexelToCornersMap, bDebugThisMapping, Light);
-
-						// Also calculate static lighting for simple light maps.  We'll force the shadows into simple light maps, but
-						// won't actually add the lights to the light guid list.  Instead, at runtime we'll check the shadow map guids
-						// for lights that are baked into light maps on platforms that don't support shadow mapping.
-						const bool bLowQualityLightMapsOnly = Light->GetDirectionalLight() == NULL;
-						CalculateDirectAreaLightingTextureMapping(TextureMapping, MappingContext, LightMapData, ShadowMaps, TexelToVertexMap, bDebugThisMapping, Light, bLowQualityLightMapsOnly);
+						if (Light->LightFlags & GI_LIGHT_USE_AREA_SHADOWS_FOR_SEPARATE_SHADOW_FACTOR)
+						{
+							FShadowMapData2D* ShadowMapData = new FShadowMapData2D(TextureMapping->CachedSizeX,TextureMapping->CachedSizeY);
+							CalculateDirectAreaLightingTextureMapping(TextureMapping, MappingContext, LightMapData, ShadowMapData, TexelToVertexMap, bDebugThisMapping, Light, false);
+
+							if (ShadowMapData)
+							{
+								FSignedDistanceFieldShadowMapData2D* ConvertedShadowMapData = new FSignedDistanceFieldShadowMapData2D(TextureMapping->CachedSizeX, TextureMapping->CachedSizeY);
+
+								for (int32 Y = 0; Y < TextureMapping->CachedSizeY; Y++)
+								{
+									for (int32 X = 0; X < TextureMapping->CachedSizeX; X++)
+									{
+										const FShadowSample& SourceShadowSample = (*ShadowMapData)(X,Y);
+										FSignedDistanceFieldShadowSample& DestShadowSample = (*ConvertedShadowMapData)(X,Y);
+										DestShadowSample.bIsMapped = SourceShadowSample.bIsMapped;
+										// Encode with more precision near 0
+										// The decode shader code will undo this in GetPrecomputedShadowMasks
+										DestShadowSample.Distance = FMath::Sqrt(FMath::Clamp(SourceShadowSample.Visibility, 0.0f, 1.0f));
+										DestShadowSample.PenumbraSize = 1;
+									}
+								}
+
+								SignedDistanceFieldShadowMaps.Add(Light, ConvertedShadowMapData);
+							}
+						}
+						else
+						{
+							const bool bUseTextureSpaceDistanceFieldGeneration = true;
+
+							if (bUseTextureSpaceDistanceFieldGeneration)
+							{
+								// Calculate distance field shadows, where the distance to the nearest shadow transition is stored instead of just a [0,1] shadow factor.
+								CalculateDirectSignedDistanceFieldLightingTextureMappingTextureSpace(TextureMapping, MappingContext, LightMapData, SignedDistanceFieldShadowMaps, TexelToVertexMap, TexelToCornersMap, bDebugThisMapping, Light);
+							}
+							else
+							{
+								// Experimental method that avoids artifacts due to lightmap seams
+								CalculateDirectSignedDistanceFieldLightingTextureMappingLightSpace(TextureMapping, MappingContext, LightMapData, SignedDistanceFieldShadowMaps, TexelToVertexMap, TexelToCornersMap, bDebugThisMapping, Light);
+							}
+						}
+
+						// Stationary directional light is never put into the lightmap, even with low quality lightmaps
+						if (!Light->GetDirectionalLight())
+						{
+							// Also calculate static lighting for simple light maps.  We'll force the shadows into simple light maps, but
+							// won't actually add the lights to the light guid list.  Instead, at runtime we'll check the shadow map guids
+							// for lights that are baked into light maps on platforms that don't support shadow mapping.
+							FShadowMapData2D* ShadowMapData = NULL;
+							const bool bLowQualityLightMapsOnly = true;
+							CalculateDirectAreaLightingTextureMapping(TextureMapping, MappingContext, LightMapData, ShadowMapData, TexelToVertexMap, bDebugThisMapping, Light, bLowQualityLightMapsOnly);
+						}
 					}
-					else
+					else if (Light->UseStaticLighting())
 					{
-						// Calculate direct lighting from area lights, no filtering in texture space.  
+						FShadowMapData2D* ShadowMapData = NULL;
+
+						// Calculate direct lighting from area lights
 						// Shadow penumbras will be correctly shaped and will be softer for larger light sources and distant shadow casters.
-						const bool bLowQualityLightMapsOnly = false;
-						CalculateDirectAreaLightingTextureMapping(TextureMapping, MappingContext, LightMapData, ShadowMaps, TexelToVertexMap, bDebugThisMapping, Light, bLowQualityLightMapsOnly);
+						CalculateDirectAreaLightingTextureMapping(TextureMapping, MappingContext, LightMapData, ShadowMapData, TexelToVertexMap, bDebugThisMapping, Light, false);
+
+						if (Light->GetMeshAreaLight() == NULL)
+						{
+							LightMapData.AddLight(Light);
+						}
 					}
 				}
 			}
@@ -938,6 +998,14 @@ void FStaticLightingSystem::SetupTextureMapping(
 						}
 					}
 				}
+				else if (GeneralSettings.bUseMaxWeight)
+				{
+					// Weighted average
+					TexelToVertex.WorldTangentX = TexelToVertex.WorldTangentX / TexelToVertex.TotalSampleWeight;
+					TexelToVertex.WorldTangentY = TexelToVertex.WorldTangentY / TexelToVertex.TotalSampleWeight;
+					TexelToVertex.WorldTangentZ = TexelToVertex.WorldTangentZ / TexelToVertex.TotalSampleWeight;
+				}
+
 				// Mark the texel as mapped to some geometry in the scene
 				CurrentLightSample.bIsMapped = true;
 
@@ -1223,7 +1291,7 @@ void FStaticLightingSystem::CalculateDirectLightingTextureMappingFiltered(
 	else
 	{
 		// Check whether the light should use a light-map or shadow-map.
-		const bool bUseStaticLighting = Light->UseStaticLighting(TextureMapping->bForceDirectLightMap);
+		const bool bUseStaticLighting = Light->UseStaticLighting();
 		if (bUseStaticLighting)
 		{
 			// Convert the shadow-map into a light-map.
@@ -1291,33 +1359,15 @@ void FStaticLightingSystem::CalculateDirectAreaLightingTextureMapping(
 	FStaticLightingTextureMapping* TextureMapping, 
 	FStaticLightingMappingContext& MappingContext,
 	FGatheredLightMapData2D& LightMapData, 
-	TMap<const FLight*, FShadowMapData2D*>& ShadowMaps,
+	FShadowMapData2D*& ShadowMapData,
 	const FTexelToVertexMap& TexelToVertexMap, 
 	bool bDebugThisMapping,
 	const FLight* Light,
-	const bool bLowQualityLightMapsOnly ) const
+	const bool bLowQualityLightMapsOnly) const
 {
 	LIGHTINGSTAT(FScopedRDTSCTimer AreaShadowsTimer(MappingContext.Stats.AreaShadowsThreadTime));
-	FShadowMapData2D* ShadowMapData = NULL;
-	const bool bUseStaticLighting =
-		Light->UseStaticLighting(TextureMapping->bForceDirectLightMap) ||
-		bLowQualityLightMapsOnly;	// Force static light maps if we're computing for LQ light maps only
-	bool bIsCompletelyOccluded = true;
-
-	if (!bUseStaticLighting)
-	{
-		// Only allow for shadow maps if shadow casting is enabled.
-		if ( (Light->LightFlags & GI_LIGHT_CASTSHADOWS) && (Light->LightFlags & GI_LIGHT_CASTSTATICSHADOWS) )
-		{
-			ShadowMapData = new FShadowMapData2D(TextureMapping->CachedSizeX,TextureMapping->CachedSizeY);
-		}
-		else
-		{
-			// Using neither static lighting nor shadow maps; nothing to do.
-			return;
-		}
-	}
 
+	bool bIsCompletelyOccluded = true;
 	FLMRandomStream SampleGenerator(0);
 
 	// Used for the optional lightmap gradient filtering pass
@@ -1353,6 +1403,8 @@ void FStaticLightingSystem::CalculateDirectAreaLightingTextureMapping(
 
 			if ( CurrentLightSample.bIsMapped )
 			{
+				UnfilteredShadowFactorData(X, Y).bIsMapped = true;
+
 				// Only continue if some part of the light is in front of the surface
 				const FTexelToVertexMap::FTexelToVertex& TexelToVertex = TexelToVertexMap(X,Y);
 
@@ -1455,9 +1507,8 @@ void FStaticLightingSystem::CalculateDirectAreaLightingTextureMapping(
 						checkSlow(TexelToVertex.TotalSampleWeight > 0.0f);
 						TransmissionCache[(Y * TextureMapping->CachedSizeX) + X] = Transmission;
 						LightIntensityCache[(Y * TextureMapping->CachedSizeX) + X] = LightIntensity;
-						UnfilteredShadowFactorData(X, Y).Visibility = ShadowFactor;
-						UnfilteredShadowFactorData(X, Y).bIsMapped = true;
-
+						// Greyscale transmission for shadowmaps
+						UnfilteredShadowFactorData(X, Y).Visibility = ShadowFactor;// * FLinearColorUtils::LinearRGBToXYZ(Transmission).G;
 						// We have valid shadow factor values, enable the filter pass
 						bShadowFactorFilterPassEnabled = true;
 					}
@@ -1474,9 +1525,9 @@ void FStaticLightingSystem::CalculateDirectAreaLightingTextureMapping(
 		const int32 KernelSizeX = 3; // Expected to be odd
 		const int32 KernelSizeY = 3; // Expected to be odd
 		const float FilterKernel3x3[KernelSizeX * KernelSizeY] = {
-			0.150f, 0.332f, 0.150f,
-			0.332f, 1.000f, 0.332f,
-			0.150f, 0.332f, 0.150f,
+			.5f * 0.150f, .5f * 0.332f, .5f * 0.150f,
+			.5f * 0.332f, .5f * 1.000f, .5f * 0.332f,
+			.5f * 0.150f, .5f * 0.332f, .5f * 0.150f,
 		};
 		for (int32 Y = 0; Y < TextureMapping->CachedSizeY; Y++)
 		{
@@ -1496,16 +1547,25 @@ void FStaticLightingSystem::CalculateDirectAreaLightingTextureMapping(
 				{
 					float UnfilteredValue = UnfilteredShadowFactorData(X, Y).Visibility;
 					const bool bIntersectingSurface = TexelToVertexMap(X, Y).bIntersectingSurface;
+					const FTexelToVertexMap::FTexelToVertex& TexelToVertex = TexelToVertexMap(X,Y);
+					const bool bLightIsInFrontOfTriangle = !Light->BehindSurface(TexelToVertex.WorldPosition, TexelToVertex.WorldTangentZ);
+
 					float FilteredValueNumerator = 0.0f;
 					float FilteredValueDenominator = 0.0f;
+					float CenterValueWeight = 1.0f;
+					
+					if (ShadowMapData)
+					{
+						// Lower the self weight on backfaces
+						// We want to spread frontface values onto backfaces for shadowmaps where the normal falloff will happen per-pixel
+						CenterValueWeight = bLightIsInFrontOfTriangle ? 1.0f : .1f;
+					}
 
 					// Compare (up to) the full grid of adjacent texels
 					int32 X1, Y1;
 					int32 FilterStepX = ((KernelSizeX - 1) / 2);
 					int32 FilterStepY = ((KernelSizeY - 1) / 2);
 
-					// Determine if filtering is needed at all
-					bool FilteringNeeded = false;
 					for (int32 KernelIndexY = -FilterStepY; KernelIndexY <= FilterStepY; KernelIndexY++)
 					{
 						// If this row is out of bounds, skip it
@@ -1528,40 +1588,41 @@ void FStaticLightingSystem::CalculateDirectAreaLightingTextureMapping(
 
 							// Only include the texel if it's not completely in shadow
 							if (UnfilteredShadowFactorData(X1, Y1).bIsMapped
+								&& !(X1 == X && Y1 == Y)
 								// Don't filter across intersecting surface boundaries
 								&& (bIntersectingSurface == TexelToVertexMap(X1, Y1).bIntersectingSurface))
 							{
 								float ComparisonValue = UnfilteredShadowFactorData(X1, Y1).Visibility;
 								float DifferenceValue = FMath::Abs(UnfilteredValue - ComparisonValue);
-								if (DifferenceValue > ThresholdForFilteringPenumbra)
-								{
-									FilteringNeeded = true;
-								}
+								const FTexelToVertexMap::FTexelToVertex& NeighborTexelToVertex = TexelToVertexMap(X1, Y1);
+								const bool bNeighborLightIsInFrontOfTriangle = !Light->BehindSurface(NeighborTexelToVertex.WorldPosition, NeighborTexelToVertex.WorldTangentZ);
 
-								// Accumulate the to-be-filtered values, in case we need to filter later.
-								// This is almost free since we're already doing the work for the lookup.
-								int32 FilterKernelIndex = ((KernelIndexY + FilterStepY) * KernelSizeX) + (KernelIndexX + FilterStepX);
-								float FilterKernelValue = FilterKernel3x3[FilterKernelIndex];
+								if (DifferenceValue > ThresholdForFilteringPenumbra 
+									// If we are filtering shadow factors for a shadowmap, only gather shadow values from frontfaces
+									&& (!ShadowMapData || bNeighborLightIsInFrontOfTriangle))
+								{
+									int32 FilterKernelIndex = ((KernelIndexY + FilterStepY) * KernelSizeX) + (KernelIndexX + FilterStepX);
+									float FilterKernelValue = FilterKernel3x3[FilterKernelIndex];
 
-								FilteredValueNumerator += (ComparisonValue * FilterKernelValue);
-								FilteredValueDenominator += FilterKernelValue;
+									FilteredValueNumerator += (ComparisonValue * FilterKernelValue);
+									FilteredValueDenominator += FilterKernelValue;
+								}
 							}
 						}
 					}
 
 					float FinalShadowFactorValue;
-					if (FilteringNeeded)
+					if (FilteredValueDenominator > 0.0f)
 					{
-						FinalShadowFactorValue = (FilteredValueNumerator / FilteredValueDenominator);
+						FinalShadowFactorValue = (FilteredValueNumerator + UnfilteredValue * CenterValueWeight) / (FilteredValueDenominator + CenterValueWeight);
 					}
 					else
 					{
 						FinalShadowFactorValue = UnfilteredValue;
 					}
 
-					// If the shadow factor value is 0.0f, it will have no effect and thus can be ignored
 					FilteredShadowFactorData(X, Y).Visibility = FinalShadowFactorValue;
-					FilteredShadowFactorData(X, Y).bIsMapped = FinalShadowFactorValue > DELTA ? true : false;
+					FilteredShadowFactorData(X, Y).bIsMapped = true;
 				}
 			}
 		}
@@ -1604,18 +1665,27 @@ void FStaticLightingSystem::CalculateDirectAreaLightingTextureMapping(
 					NumUnoccludedTexels++;
 					// Get any cached values
 					const float AdjustedShadowFactor = FMath::Pow(ShadowFactor, Light->ShadowExponent);
-					const FLinearColor Transmission = TransmissionCache[(Y * TextureMapping->CachedSizeX) + X];
-					const FLinearColor LightIntensity = LightIntensityCache[(Y * TextureMapping->CachedSizeX) + X];
-
-					// Calculate any derived values
-					const FTexelToVertexMap::FTexelToVertex& TexelToVertex = TexelToVertexMap(X,Y);
-					const FStaticLightingVertex CurrentVertex = TexelToVertex.GetVertex();
-					const FGatheredLightSample DirectLighting = CalculatePointLighting(TextureMapping, CurrentVertex, TexelToVertex.ElementIndex, Light, LightIntensity, Transmission);
 
 					if (GeneralSettings.ViewSingleBounceNumber < 1)
 					{
-						if (bUseStaticLighting)
+						if (ShadowMapData)
+						{
+							FShadowSample& CurrentShadowSample = (*ShadowMapData)(X,Y);
+							CurrentShadowSample.Visibility = AdjustedShadowFactor;
+							if ( CurrentShadowSample.Visibility > 0.0001f )
+							{
+								bIsCompletelyOccluded = false;
+							}
+						}
+						else
 						{
+							// Calculate any derived values
+							const FTexelToVertexMap::FTexelToVertex& TexelToVertex = TexelToVertexMap(X,Y);
+							const FStaticLightingVertex CurrentVertex = TexelToVertex.GetVertex();
+							const FLinearColor LightIntensity = LightIntensityCache[(Y * TextureMapping->CachedSizeX) + X];
+							const FLinearColor Transmission = TransmissionCache[(Y * TextureMapping->CachedSizeX) + X];
+							const FGatheredLightSample DirectLighting = CalculatePointLighting(TextureMapping, CurrentVertex, TexelToVertex.ElementIndex, Light, LightIntensity, Transmission);
+
 							FGatheredLightMapSample& CurrentLightSample = LightMapData(X,Y);
 							if( bLowQualityLightMapsOnly )
 							{
@@ -1626,46 +1696,16 @@ void FStaticLightingSystem::CalculateDirectAreaLightingTextureMapping(
 								CurrentLightSample.AddWeighted(DirectLighting, AdjustedShadowFactor);
 							}
 						}
-						else
-						{
-							FShadowSample& CurrentShadowSample = (*ShadowMapData)(X,Y);
-							// Using greyscale transmission for shadow maps, since we don't want to increase storage
-							CurrentShadowSample.Visibility = AdjustedShadowFactor * FLinearColorUtils::LinearRGBToXYZ(Transmission).G;
-							if ( CurrentShadowSample.Visibility > 0.0001f )
-							{
-								bIsCompletelyOccluded = false;
-							}
-						}
 					}
 				}
 			}
 		}
 	}
 
-	if (bUseStaticLighting)
+	if (ShadowMapData && (bIsCompletelyOccluded || NumUnoccludedTexels < NumMappedTexels * ShadowSettings.MinUnoccludedFraction))
 	{
-		// Don't modify light array when only updating LQ light maps.  This light has already been added
-		// to the shadow map light set, and we never want a light to appear in both lists.  Instead, the 
-		// runtime code will check both the light and shadow map lists for lights when running in a mode
-		// that does not support shadow maps.
-		if( !bLowQualityLightMapsOnly )
-		{
-			if (Light->GetMeshAreaLight() == NULL)
-			{
-				LightMapData.AddLight(Light);
-			}
-		}
-	}
-	else if (ShadowMapData)
-	{
-		if (bIsCompletelyOccluded || NumUnoccludedTexels < NumMappedTexels * ShadowSettings.MinUnoccludedFraction)
-		{
-			delete ShadowMapData;
-		}
-		else
-		{
-			ShadowMaps.Add(Light,ShadowMapData);
-		}
+		delete ShadowMapData;
+		ShadowMapData = NULL;
 	}
 }
 
@@ -1811,7 +1851,7 @@ void FDistanceFieldRasterPolicy::ProcessPixel(int32 X,int32 Y,const InterpolantT
  * Calculate signed distance field shadowing from a single light,  
  * Based on the paper "Improved Alpha-Tested Magnification for Vector Textures and Special Effects" by Valve.
  */
-void FStaticLightingSystem::CalculateDirectSignedDistanceFieldLightingTextureMapping(
+void FStaticLightingSystem::CalculateDirectSignedDistanceFieldLightingTextureMappingTextureSpace(
 	FStaticLightingTextureMapping* TextureMapping, 
 	FStaticLightingMappingContext& MappingContext,
 	FGatheredLightMapData2D& LightMapData, 
@@ -2006,6 +2046,7 @@ void FStaticLightingSystem::CalculateDirectSignedDistanceFieldLightingTextureMap
 					{
 						// Initialize the final distance field data, since it will only be written to after this if it gets scattered to during the search.
 						FinalShadowSample.Distance = 1.0f;
+						FinalShadowSample.PenumbraSize = 1.0f;
 					}
 
 					// Search for a neighbor with different visibility
@@ -2416,6 +2457,254 @@ void FStaticLightingSystem::CalculateDirectSignedDistanceFieldLightingTextureMap
 	}
 }
 
+void FStaticLightingSystem::CalculateDirectSignedDistanceFieldLightingTextureMappingLightSpace(
+	FStaticLightingTextureMapping* TextureMapping, 
+	FStaticLightingMappingContext& MappingContext,
+	FGatheredLightMapData2D& LightMapData, 
+	TMap<const FLight*, FSignedDistanceFieldShadowMapData2D*>& ShadowMaps,
+	const FTexelToVertexMap& TexelToVertexMap, 
+	const FTexelToCornersMap& TexelToCornersMap,
+	bool bDebugThisMapping,
+	const FLight* Light) const
+{
+	const FBoxSphereBounds MeshInfluenceBounds(TextureMapping->Mesh->BoundingBox.ExpandBy(ShadowSettings.MaxTransitionDistanceWorldSpace));
+
+	if (Light->AffectsBounds(MeshInfluenceBounds))
+	{
+		const FBoxSphereBounds SceneBounds = FBoxSphereBounds(AggregateMesh.GetBounds());
+		const FDirectionalLight* DirectionalLight = Light->GetDirectionalLight();
+		const FSpotLight* SpotLight = Light->GetSpotLight();
+		const FPointLight* PointLight = Light->GetPointLight();
+		check(DirectionalLight || SpotLight || PointLight);
+
+		if (DirectionalLight)
+		{
+			LIGHTINGSTAT(FManualRDTSCTimer FirstPassSourceTimer(MappingContext.Stats.SignedDistanceFieldSourceFirstPassThreadTime));
+			FVector4 XAxis, YAxis;
+			DirectionalLight->Direction.FindBestAxisVectors3(XAxis, YAxis);
+			// Create a coordinate system for the directional light, with the z axis corresponding to the light's direction
+			FMatrix WorldToLight = FBasisVectorMatrix(XAxis, YAxis, DirectionalLight->Direction, FVector4(0,0,0));
+
+			const FBox LightSpaceImportanceBounds = MeshInfluenceBounds.GetBox().TransformBy(WorldToLight);
+
+			FStaticShadowDepthMap ShadowDepthMap;
+
+			int32 ShadowMapSizeX = FMath::TruncToInt(FMath::Max(LightSpaceImportanceBounds.GetExtent().X * 2.0f * 100.0f / ShadowSettings.MaxTransitionDistanceWorldSpace, 4.0f));
+			ShadowMapSizeX = ShadowMapSizeX == appTruncErrorCode ? INT_MAX : ShadowMapSizeX;
+			int32 ShadowMapSizeY = FMath::TruncToInt(FMath::Max(LightSpaceImportanceBounds.GetExtent().Y * 2.0f * 100.0f / ShadowSettings.MaxTransitionDistanceWorldSpace, 4.0f));
+			ShadowMapSizeY = ShadowMapSizeY == appTruncErrorCode ? INT_MAX : ShadowMapSizeY;
+
+			uint64 ShadowDepthMapMaxSamples = 4194304;
+
+			// Clamp the number of dominant shadow samples generated if necessary while maintaining aspect ratio
+			if ((uint64)ShadowMapSizeX * (uint64)ShadowMapSizeY > (uint64)ShadowDepthMapMaxSamples)
+			{
+				const float AspectRatio = ShadowMapSizeX / (float)ShadowMapSizeY;
+				ShadowMapSizeY = FMath::TruncToInt(FMath::Sqrt(ShadowDepthMapMaxSamples / AspectRatio));
+				ShadowMapSizeX = FMath::TruncToInt(ShadowDepthMapMaxSamples / ShadowMapSizeY);
+			}
+
+			TArray<float> ShadowMap;
+
+			// Allocate the shadow map
+			ShadowMap.Empty(ShadowMapSizeX * ShadowMapSizeY);
+			ShadowMap.AddZeroed(ShadowMapSizeX * ShadowMapSizeY);
+			const float ShadowMapStart = LightSpaceImportanceBounds.Max.Z - SceneBounds.SphereRadius * 2;
+			const FMatrix LightToWorld = WorldToLight.InverseFast();
+
+			for (int32 Y = 0; Y < ShadowMapSizeY; Y++)
+			{
+				const float YFraction = (Y + .5f) / (float)(ShadowMapSizeY - 1);
+
+				for (int32 X = 0; X < ShadowMapSizeX; X++)
+				{
+					const float XFraction = (X + .5f) / (float)(ShadowMapSizeX - 1);
+
+					const FVector4 LightSpaceEndPosition(
+						LightSpaceImportanceBounds.Min.X + XFraction * (LightSpaceImportanceBounds.Max.X - LightSpaceImportanceBounds.Min.X),
+						LightSpaceImportanceBounds.Min.Y + YFraction * (LightSpaceImportanceBounds.Max.Y - LightSpaceImportanceBounds.Min.Y),
+						LightSpaceImportanceBounds.Max.Z);
+					const FVector4 WorldSpaceEndPosition = LightToWorld.TransformPosition(LightSpaceEndPosition);
+
+					const FVector4 LightSpaceStartPosition(LightSpaceEndPosition.X, LightSpaceEndPosition.Y, ShadowMapStart);
+					const FVector4 WorldSpaceStartPosition = LightToWorld.TransformPosition(LightSpaceStartPosition);
+
+					const FLightRay LightRay(
+						WorldSpaceStartPosition,
+						WorldSpaceEndPosition,
+						NULL,
+						NULL,
+						// We are tracing from the light instead of to the light,
+						// So flip sidedness so that backface culling matches up with tracing to the light
+						LIGHTRAY_FLIP_SIDEDNESS
+						);
+
+					FLightRayIntersection Intersection;
+					AggregateMesh.IntersectLightRay(LightRay, true, false, true, MappingContext.RayCache, Intersection);
+
+					float MaxSampleDistance = SceneBounds.SphereRadius * 2;
+
+					if (Intersection.bIntersects)
+					{
+						MaxSampleDistance = (Intersection.IntersectionVertex.WorldPosition - WorldSpaceStartPosition).Size3();
+					}
+
+					ShadowMap[Y * ShadowMapSizeX + X] = MaxSampleDistance;
+				}
+			}
+
+			FirstPassSourceTimer.Stop();
+
+			LIGHTINGSTAT(FScopedRDTSCTimer SearchTimer(MappingContext.Stats.SignedDistanceFieldSearchThreadTime));
+			FSignedDistanceFieldShadowMapData2D* ShadowMapData = new FSignedDistanceFieldShadowMapData2D(TextureMapping->CachedSizeX, TextureMapping->CachedSizeY);
+			const int32 TransitionSearchTexelRadiusX = FMath::TruncToInt(ShadowMapSizeX * ShadowSettings.MaxTransitionDistanceWorldSpace / LightSpaceImportanceBounds.GetSize().X);
+			const int32 TransitionSearchTexelRadiusY = FMath::TruncToInt(ShadowMapSizeY * ShadowSettings.MaxTransitionDistanceWorldSpace / LightSpaceImportanceBounds.GetSize().Y);
+			const float BoundsCellSizeX = (LightSpaceImportanceBounds.Max.X - LightSpaceImportanceBounds.Min.X) / ShadowMapSizeX;
+			const float BoundsCellSizeY = (LightSpaceImportanceBounds.Max.Y - LightSpaceImportanceBounds.Min.Y) / ShadowMapSizeY;
+			const float DepthBias = FMath::Max(BoundsCellSizeX, BoundsCellSizeY);
+
+			for (int32 Y = 0; Y < TextureMapping->CachedSizeY; Y++)
+			{
+				for (int32 X = 0; X < TextureMapping->CachedSizeX; X++)
+				{
+					bool bDebugThisTexel = false;
+#if ALLOW_LIGHTMAP_SAMPLE_DEBUGGING
+					if (bDebugThisMapping
+						&& Y == Scene.DebugInput.LocalY
+						&& X == Scene.DebugInput.LocalX)
+					{
+						bDebugThisTexel = true;
+					}
+#endif
+					const FTexelToVertexMap::FTexelToVertex& TexelToVertex = TexelToVertexMap(X,Y);
+
+					if (TexelToVertex.TotalSampleWeight > 0.0f)
+					{
+						const FVector4 LightSpacePosition = WorldToLight.TransformPosition(TexelToVertex.WorldPosition);
+						const FVector LightSpaceNormal = WorldToLight.TransformVector(TexelToVertex.WorldTangentZ);
+						const float SinThetaX = LightSpaceNormal.X;
+						const float TanThetaX = SinThetaX / FMath::Sqrt(1 - SinThetaX * SinThetaX);
+						const float SinThetaY = LightSpaceNormal.Y;
+						const float TanThetaY = SinThetaY / FMath::Sqrt(1 - SinThetaY * SinThetaY);
+						const float SurfaceDepth = LightSpacePosition.Z - ShadowMapStart;
+
+						const int32 ShadowMapX = FMath::Clamp(FMath::TruncToInt((LightSpacePosition.X - LightSpaceImportanceBounds.Min.X) / BoundsCellSizeX), 0, ShadowMapSizeX - 1);
+						const int32 ShadowMapY = FMath::Clamp(FMath::TruncToInt((LightSpacePosition.Y - LightSpaceImportanceBounds.Min.Y) / BoundsCellSizeY), 0, ShadowMapSizeY - 1);
+
+						const float TexelShadowMapDepth = ShadowMap[ShadowMapY * ShadowMapSizeX + ShadowMapX];
+						const float SlopeScaledDepthBias = 4 * FMath::Max(BoundsCellSizeX * FMath::Abs(TanThetaX), BoundsCellSizeY * FMath::Abs(TanThetaY));
+						const bool bTexelVisible = TexelShadowMapDepth > SurfaceDepth - SlopeScaledDepthBias - DepthBias;
+						float ClosestTransition = ShadowSettings.MaxTransitionDistanceWorldSpace;
+						//float ClosestTransitionPenumbraSize = 1;
+						float MostShadowingTransition = 1;
+						float MostShadowingTransitionDistance = 1;
+						float MostShadowingTransitionPenumbraSize = 1;
+
+						for (int32 SearchY = FMath::Max(ShadowMapY - TransitionSearchTexelRadiusY, 0); SearchY < FMath::Min(ShadowMapY + TransitionSearchTexelRadiusY, ShadowMapSizeY); SearchY++)
+						{
+							for (int32 SearchX = FMath::Max(ShadowMapX - TransitionSearchTexelRadiusX, 0); SearchX < FMath::Min(ShadowMapX + TransitionSearchTexelRadiusX, ShadowMapSizeX); SearchX++)
+							{
+								const FVector2D LightSpaceXYOffset((SearchX - ShadowMapX) * BoundsCellSizeX, (SearchY - ShadowMapY) * BoundsCellSizeY);
+								const float PlaneHeightOffsetX = LightSpaceXYOffset.X * TanThetaX;
+								const float PlaneHeightOffsetY = LightSpaceXYOffset.Y * TanThetaY;
+								const float PlaneHeightOffset = PlaneHeightOffsetX + PlaneHeightOffsetY;
+
+								const float SearchShadowMapDepth = ShadowMap[SearchY * ShadowMapSizeX + SearchX];
+								const float ExtrapolatedSurfaceDepth = SurfaceDepth + PlaneHeightOffset;
+								const float SearchTransition = LightSpaceXYOffset.Size();
+								const float SameSurfaceDepthBias = SearchTransition * 1.0f;
+								const bool bSearchTexelVisible = SearchShadowMapDepth > ExtrapolatedSurfaceDepth - SlopeScaledDepthBias - DepthBias - SameSurfaceDepthBias;
+
+								if (bTexelVisible)
+								{
+									const float SearchNormalizedDistance = FMath::Clamp(SearchTransition / ShadowSettings.MaxTransitionDistanceWorldSpace, 0.0f, 1.0f);
+									const float SearchEncodedDistance = bTexelVisible ? (SearchNormalizedDistance) * .5f + .5f : .5f - SearchNormalizedDistance * .5f;
+
+									const float ReceiverDistanceFromLight = SurfaceDepth;
+									const float OccluderDistanceFromLight = bTexelVisible ? SearchShadowMapDepth : TexelShadowMapDepth;
+									// World space distance from center of penumbra to fully shadowed or fully lit transition
+									const float SearchPenumbraSize = (ReceiverDistanceFromLight - OccluderDistanceFromLight) * Light->LightSourceRadius / OccluderDistanceFromLight;
+									const float SearchEncodedPenumbraSize = FMath::Clamp(SearchPenumbraSize / ShadowSettings.MaxTransitionDistanceWorldSpace, 0.01f, 1.0f);
+
+									const float SearchShadowing = FMath::Clamp(SearchEncodedDistance / SearchEncodedPenumbraSize - .5f / SearchEncodedPenumbraSize + .5f, 0.0f, 1.0f);
+
+									if (bSearchTexelVisible != bTexelVisible
+										/* && SearchTransition < ClosestTransition*/
+										&& SearchShadowing < MostShadowingTransition)
+									{
+										MostShadowingTransition = SearchShadowing;
+										MostShadowingTransitionDistance = SearchEncodedDistance;
+										MostShadowingTransitionPenumbraSize = SearchEncodedPenumbraSize;
+										//ClosestTransition = SearchTransition;
+										/*
+										if (bTexelVisible)
+										{
+											// Approximate the penumbra size using PenumbraSize = (ReceiverDistanceFromLight - OccluderDistanceFromLight) * LightSize / OccluderDistanceFromLight,
+											// Which is from the paper "Percentage-Closer Soft Shadows" by Randima Fernando
+											const float ReceiverDistanceFromLight = SurfaceDepth;
+											const float OccluderDistanceFromLight = SearchShadowMapDepth;
+											// World space distance from center of penumbra to fully shadowed or fully lit transition
+											const float PenumbraSize = (ReceiverDistanceFromLight - OccluderDistanceFromLight) * Light->LightSourceRadius / OccluderDistanceFromLight;
+											// Normalize the penumbra size so it is a fraction of MaxTransitionDistanceWorldSpace
+											ClosestTransitionPenumbraSize = FMath::Clamp(PenumbraSize / ShadowSettings.MaxTransitionDistanceWorldSpace, 0.01f, 1.0f);
+										}*/
+									}
+								}
+								else
+								{
+									if (bSearchTexelVisible != bTexelVisible
+										&& SearchTransition < ClosestTransition)
+									{
+										ClosestTransition = SearchTransition;
+									}
+								}
+							}
+						}
+
+						FSignedDistanceFieldShadowSample& FinalShadowSample = (*ShadowMapData)(X, Y);
+						FinalShadowSample.bIsMapped = true;
+						FinalShadowSample.Distance = MostShadowingTransitionDistance;
+						FinalShadowSample.PenumbraSize = MostShadowingTransitionPenumbraSize;
+
+						if (!bTexelVisible)
+						{
+							const float NormalizedDistance = FMath::Clamp(ClosestTransition / ShadowSettings.MaxTransitionDistanceWorldSpace, 0.0f, 1.0f);
+							// Encode the transition distance so that [.5,0] corresponds to [0,1] for shadowed texels, and [.5,1] corresponds to [0,1] for unshadowed texels.
+							// .5 of the encoded distance lies exactly on the shadow transition.
+							FinalShadowSample.Distance = bTexelVisible ? (NormalizedDistance) * .5f + .5f : .5f - NormalizedDistance * .5f;
+
+							const float ReceiverDistanceFromLight = SurfaceDepth;
+							const float OccluderDistanceFromLight = TexelShadowMapDepth;
+							const float PenumbraSize = (ReceiverDistanceFromLight - OccluderDistanceFromLight) * Light->LightSourceRadius / OccluderDistanceFromLight;
+							FinalShadowSample.PenumbraSize = FMath::Clamp(PenumbraSize / ShadowSettings.MaxTransitionDistanceWorldSpace, 0.01f, 1.0f);
+						}
+						/*
+						const float NormalizedDistance = FMath::Clamp(ClosestTransition / ShadowSettings.MaxTransitionDistanceWorldSpace, 0.0f, 1.0f);
+						// Encode the transition distance so that [.5,0] corresponds to [0,1] for shadowed texels, and [.5,1] corresponds to [0,1] for unshadowed texels.
+						// .5 of the encoded distance lies exactly on the shadow transition.
+						FinalShadowSample.Distance = bTexelVisible ? (NormalizedDistance) * .5f + .5f : .5f - NormalizedDistance * .5f;
+
+						if (bTexelVisible)
+						{
+							FinalShadowSample.PenumbraSize = ClosestTransitionPenumbraSize;
+						}
+						else
+						{
+							const float ReceiverDistanceFromLight = SurfaceDepth;
+							const float OccluderDistanceFromLight = TexelShadowMapDepth;
+							const float PenumbraSize = (ReceiverDistanceFromLight - OccluderDistanceFromLight) * Light->LightSourceRadius / OccluderDistanceFromLight;
+							FinalShadowSample.PenumbraSize = FMath::Clamp(PenumbraSize / ShadowSettings.MaxTransitionDistanceWorldSpace, 0.01f, 1.0f);
+						}
+						*/
+					}
+				}
+			}
+
+			ShadowMaps.Add(Light, ShadowMapData);
+		}
+	}
+}
+
 /**
  * Estimate direct lighting using the direct photon map.
  * This is only useful for debugging what the final gather rays see.
diff --git a/Engine/Source/Programs/UnrealLightmass/Public/SceneExport.h b/Engine/Source/Programs/UnrealLightmass/Public/SceneExport.h
--- a/Engine/Source/Programs/UnrealLightmass/Public/SceneExport.h
+++ b/Engine/Source/Programs/UnrealLightmass/Public/SceneExport.h
@@ -758,15 +758,15 @@ enum EDawnLightFlags
 	GI_LIGHT_CASTSHADOWS			= 0x00000001,
 	// maps to ULightComponent::HasStaticLighting()
 	GI_LIGHT_HASSTATICLIGHTING		= 0x00000002,
-	// maps to ULightComponent::UseDirectLightMap
-	GI_LIGHT_USEDIRECTLIGHTMAP		= 0x00000004,
 	// maps to ULightComponent::HasStaticShadowing()
 	GI_LIGHT_HASSTATICSHADOWING		= 0x00000008,
 	// maps to ULightComponent::CastStaticShadows
 	GI_LIGHT_CASTSTATICSHADOWS		= 0x00000010,
-	GI_LIGHT_USESIGNEDDISTANCEFIELDSHADOWS = 0x00000020,
+	GI_LIGHT_STORE_SEPARATE_SHADOW_FACTOR = 0x00000020,
 	GI_LIGHT_INVERSE_SQUARED		= 0x00000080,
-	GI_LIGHT_USE_LIGHTPROFILE		= 0x00000100
+	GI_LIGHT_USE_LIGHTPROFILE		= 0x00000100,
+	// Whether a stationary light should generate a standard shadowmap (area shadows) or a distance field shadow map
+	GI_LIGHT_USE_AREA_SHADOWS_FOR_SEPARATE_SHADOW_FACTOR		= 0x00000200
 };
 
 struct FLightData
@@ -1042,8 +1042,6 @@ struct FStaticLightingMappingData
 {
 	FGuid Guid;
 	FGuid StaticLightingMeshInstance;
-	/** true if light-maps to be used for the object's direct lighting. */
-	bool bForceDirectLightMap;
 };
 
 struct FStaticLightingTextureMappingData
diff --git a/Engine/Source/Runtime/Core/Private/HAL/ConsoleManager.cpp b/Engine/Source/Runtime/Core/Private/HAL/ConsoleManager.cpp
--- a/Engine/Source/Runtime/Core/Private/HAL/ConsoleManager.cpp
+++ b/Engine/Source/Runtime/Core/Private/HAL/ConsoleManager.cpp
@@ -1872,13 +1872,6 @@ static TAutoConsoleVariable<int32> CVarRenderTargetPoolMin(
 	TEXT("Default is 200 MB."),
 	ECVF_Scalability | ECVF_RenderThreadSafe);
 
-static TAutoConsoleVariable<float> CVarShadowDistanceFieldPenumbraSize(
-	TEXT("r.Shadow.DistanceFieldPenumbraSize"),
-	0.05f,
-	TEXT("Controls the size of the uniform penumbra produced by static shadowing.\n")
-	TEXT("This is a fraction of MaxTransitionDistanceWorldSpace in BaseLightmass.ini."),
-	ECVF_RenderThreadSafe);
-
 static TAutoConsoleVariable<int32> CVarIdleWhenNotForeground(
 	TEXT("t.IdleWhenNotForeground"), 0,
 	TEXT("Prevents the engine from taking any CPU or GPU time while not the foreground app."),
diff --git a/Engine/Source/Runtime/Core/Public/UObject/ObjectVersion.h b/Engine/Source/Runtime/Core/Public/UObject/ObjectVersion.h
--- a/Engine/Source/Runtime/Core/Public/UObject/ObjectVersion.h
+++ b/Engine/Source/Runtime/Core/Public/UObject/ObjectVersion.h
@@ -530,6 +530,8 @@ enum EUnrealEngineObjectUE4Version
 	VER_UE4_ADDED_NATIVE_SERIALIZATION_FOR_IMMUTABLE_STRUCTURES,
 	// Deprecated attributes that override the style on UMG widgets
 	VER_UE4_DEPRECATE_UMG_STYLE_OVERRIDES,
+	// Shadowmap penumbra size stored
+	VER_UE4_STATIC_SHADOWMAP_PENUMBRA_SIZE,
 
 	// -----<new versions can be added before this line>-------------------------------------------------
 	// - this needs to be the last line (see note below)
diff --git a/Engine/Source/Runtime/Engine/Classes/Components/DirectionalLightComponent.h b/Engine/Source/Runtime/Engine/Classes/Components/DirectionalLightComponent.h
--- a/Engine/Source/Runtime/Engine/Classes/Components/DirectionalLightComponent.h
+++ b/Engine/Source/Runtime/Engine/Classes/Components/DirectionalLightComponent.h
@@ -152,6 +152,13 @@ class UDirectionalLightComponent : public ULightComponent
 	// ULightComponent interface.
 	virtual FVector4 GetLightPosition() const override;
 	virtual ELightComponentType GetLightType() const override;
+	virtual FLightmassLightSettings GetLightmassSettings() const override
+	{
+		return LightmassSettings;
+	}
+
+	virtual float GetUniformPenumbraSize() const override;
+
 	virtual FLightSceneProxy* CreateSceneProxy() const override;
 	virtual bool IsUsedAsAtmosphereSunLight() const override
 	{
diff --git a/Engine/Source/Runtime/Engine/Classes/Components/LightComponent.h b/Engine/Source/Runtime/Engine/Classes/Components/LightComponent.h
--- a/Engine/Source/Runtime/Engine/Classes/Components/LightComponent.h
+++ b/Engine/Source/Runtime/Engine/Classes/Components/LightComponent.h
@@ -272,7 +272,9 @@ public:
 	*/
 	virtual ELightComponentType GetLightType() const PURE_VIRTUAL(ULightComponent::GetLightType,return LightType_MAX;);
 
+	virtual FLightmassLightSettings GetLightmassSettings() const PURE_VIRTUAL(ULightComponent::GetLightmassSettings,return FLightmassLightSettings(););
 
+	virtual float GetUniformPenumbraSize() const PURE_VIRTUAL(ULightComponent::GetUniformPenumbraSize,return 0;);
 
 	/**
 	 * Update/reset light GUIDs.
diff --git a/Engine/Source/Runtime/Engine/Classes/Components/PointLightComponent.h b/Engine/Source/Runtime/Engine/Classes/Components/PointLightComponent.h
--- a/Engine/Source/Runtime/Engine/Classes/Components/PointLightComponent.h
+++ b/Engine/Source/Runtime/Engine/Classes/Components/PointLightComponent.h
@@ -83,6 +83,13 @@ public:
 	virtual FBox GetBoundingBox() const override;
 	virtual FSphere GetBoundingSphere() const override;
 	virtual ELightComponentType GetLightType() const override;
+	virtual FLightmassLightSettings GetLightmassSettings() const override
+	{
+		return LightmassSettings;
+	}
+
+	virtual float GetUniformPenumbraSize() const override;
+
 	virtual FLightSceneProxy* CreateSceneProxy() const override;
 
 	// Begin UObject Interface
diff --git a/Engine/Source/Runtime/Engine/Classes/Engine/EngineTypes.h b/Engine/Source/Runtime/Engine/Classes/Engine/EngineTypes.h
--- a/Engine/Source/Runtime/Engine/Classes/Engine/EngineTypes.h
+++ b/Engine/Source/Runtime/Engine/Classes/Engine/EngineTypes.h
@@ -1219,9 +1219,17 @@ struct FLightmassLightSettings
 	UPROPERTY(EditAnywhere, BlueprintReadWrite, Category=Lightmass, meta=(UIMin = "0.1", UIMax = "4.0"))
 	float ShadowExponent;
 
+	/** 
+	 * Whether to use area shadows for stationary light precomputed shadowmaps.  
+	 * Area shadows get softer the further they are from shadow casters, but require higher lightmap resolution to get the same quality where the shadow is sharp.
+	 */
+	UPROPERTY(EditAnywhere, BlueprintReadWrite, Category=Lightmass)
+	bool bUseAreaShadowsForStationaryLight;
+
 	FLightmassLightSettings()
 		: IndirectLightingSaturation(1.0f)
 		, ShadowExponent(2.0f)
+		, bUseAreaShadowsForStationaryLight(false)
 	{ }
 };
 
diff --git a/Engine/Source/Runtime/Engine/Classes/Engine/RendererSettings.h b/Engine/Source/Runtime/Engine/Classes/Engine/RendererSettings.h
--- a/Engine/Source/Runtime/Engine/Classes/Engine/RendererSettings.h
+++ b/Engine/Source/Runtime/Engine/Classes/Engine/RendererSettings.h
@@ -169,11 +169,6 @@ class ENGINE_API URendererSettings : public UDeveloperSettings
 		ToolTip = "Whether to generate a low-resolution base color texture for landscapes for rendering real-time global illumination.  This feature requires GenerateMeshDistanceFields is also enabled, and will increase mesh build times and memory usage."))
 		uint32 bGenerateLandscapeGIData : 1;
 
-	UPROPERTY(config, EditAnywhere, Category=Lighting, meta=(
-		ConsoleVariable="r.Shadow.DistanceFieldPenumbraSize",
-		ToolTip="Controls the size of the uniform penumbra produced by static shadowing."))
-	float DistanceFieldPenumbraSize;
-
 	UPROPERTY(config, EditAnywhere, Category=Tessellation, meta=(
 		ConsoleVariable="r.TessellationAdaptivePixelsPerTriangle",DisplayName="Adaptive pixels per triangle",
 		ToolTip="When adaptive tessellation is enabled it will try to tessellate a mesh so that each triangle contains the specified number of pixels. The tessellation multiplier specified in the material can increase or decrease the amount of tessellation."))
diff --git a/Engine/Source/Runtime/Engine/Private/Components/DirectionalLightComponent.cpp b/Engine/Source/Runtime/Engine/Private/Components/DirectionalLightComponent.cpp
--- a/Engine/Source/Runtime/Engine/Private/Components/DirectionalLightComponent.cpp
+++ b/Engine/Source/Runtime/Engine/Private/Components/DirectionalLightComponent.cpp
@@ -783,6 +783,20 @@ ELightComponentType UDirectionalLightComponent::GetLightType() const
 	return LightType_Directional;
 }
 
+float UDirectionalLightComponent::GetUniformPenumbraSize() const
+{
+	if (LightmassSettings.bUseAreaShadowsForStationaryLight)
+	{
+		// Interpret distance as shadow factor directly
+		return 1.0f;
+	}
+	else
+	{
+		// Heuristic to derive uniform penumbra size from light source angle
+		return FMath::Clamp(LightmassSettings.LightSourceAngle * .05f, .0001f, 1.0f);
+	}
+}
+
 void UDirectionalLightComponent::SetDynamicShadowDistanceMovableLight(float NewValue)
 {
 	if (DynamicShadowDistanceMovableLight != NewValue)
diff --git a/Engine/Source/Runtime/Engine/Private/Components/PointLightComponent.cpp b/Engine/Source/Runtime/Engine/Private/Components/PointLightComponent.cpp
--- a/Engine/Source/Runtime/Engine/Private/Components/PointLightComponent.cpp
+++ b/Engine/Source/Runtime/Engine/Private/Components/PointLightComponent.cpp
@@ -205,6 +205,20 @@ ELightComponentType UPointLightComponent::GetLightType() const
 	return LightType_Point;
 }
 
+float UPointLightComponent::GetUniformPenumbraSize() const
+{
+	if (LightmassSettings.bUseAreaShadowsForStationaryLight)
+	{
+		// Interpret distance as shadow factor directly
+		return 1.0f;
+	}
+	else
+	{
+		// Heuristic to derive uniform penumbra size from light source radius
+		return FMath::Clamp(SourceRadius == 0 ? .05f : SourceRadius * .005f, .0001f, 1.0f);
+	}
+}
+
 //
 FBox UPointLightComponent::GetBoundingBox() const
 {
diff --git a/Engine/Source/Runtime/Engine/Private/ShadowMap.cpp b/Engine/Source/Runtime/Engine/Private/ShadowMap.cpp
--- a/Engine/Source/Runtime/Engine/Private/ShadowMap.cpp
+++ b/Engine/Source/Runtime/Engine/Private/ShadowMap.cpp
@@ -294,31 +294,37 @@ void FShadowMapPendingTexture::StartEncoding(UWorld* InWorld)
 	Texture->SRGB			= false;
 	Texture->LODGroup		= TEXTUREGROUP_Shadowmap;
 	Texture->ShadowmapFlags	= ShadowmapFlags;
-	Texture->Source.Init2DWithMipChain(GetSizeX(), GetSizeY(), TSF_BGRA8);
-	Texture->MipGenSettings = TMGS_LeaveExistingMips;
-	Texture->CompressionNone = true;
-
-	int32 TextureSizeX = Texture->Source.GetSizeX();
-	int32 TextureSizeY = Texture->Source.GetSizeY();
 
 	{
 		// Create the uncompressed top mip-level.
 		TArray< TArray<FFourDistanceFieldSamples> > MipData;
-		FShadowMap2D::EncodeSingleTexture(*this, Texture, MipData);
+		const int32 NumChannelsUsed = FShadowMap2D::EncodeSingleTexture(*this, Texture, MipData);
+
+		Texture->Source.Init2DWithMipChain(GetSizeX(), GetSizeY(), NumChannelsUsed == 1 ? TSF_G8 : TSF_BGRA8);
+		Texture->MipGenSettings = TMGS_LeaveExistingMips;
+		Texture->CompressionNone = true;
 
 		// Copy the mip-map data into the UShadowMapTexture2D's mip-map array.
 		for(int32 MipIndex = 0;MipIndex < MipData.Num();MipIndex++)
 		{
-			FColor* DestMipData = (FColor*)Texture->Source.LockMip(MipIndex);
-			uint32 MipSizeX = FMath::Max(1,TextureSizeX >> MipIndex);
-			uint32 MipSizeY = FMath::Max(1,TextureSizeY >> MipIndex);
+			uint8* DestMipData = (uint8*)Texture->Source.LockMip(MipIndex);
+			uint32 MipSizeX = FMath::Max<uint32>(1, GetSizeX() >> MipIndex);
+			uint32 MipSizeY = FMath::Max<uint32>(1, GetSizeY() >> MipIndex);
 
 			for(uint32 Y = 0;Y < MipSizeY;Y++)
 			{
 				for(uint32 X = 0;X < MipSizeX;X++)
 				{
 					const FFourDistanceFieldSamples& SourceSample = MipData[MipIndex][Y * MipSizeX + X];
-					DestMipData[ Y * MipSizeX + X ] = FColor(SourceSample.Samples[0].Distance, SourceSample.Samples[1].Distance, SourceSample.Samples[2].Distance, SourceSample.Samples[3].Distance);
+
+					if (NumChannelsUsed == 1)
+					{
+						DestMipData[Y * MipSizeX + X] = SourceSample.Samples[0].Distance;
+					}
+					else
+					{
+						((FColor*)DestMipData)[Y * MipSizeX + X] = FColor(SourceSample.Samples[0].Distance, SourceSample.Samples[1].Distance, SourceSample.Samples[2].Distance, SourceSample.Samples[3].Distance);
+					}
 				}
 			}
 
@@ -333,7 +339,7 @@ void FShadowMapPendingTexture::StartEncoding(UWorld* InWorld)
 
 	// Update stats.
 	int32 TextureSize			= Texture->CalcTextureMemorySizeEnum( TMC_AllMips );
-	GNumShadowmapTotalTexels	+= TextureSizeX * TextureSizeY;
+	GNumShadowmapTotalTexels	+= GetSizeX() * GetSizeY();
 	GNumShadowmapTextures++;
 	GShadowmapTotalSize			+= TextureSize;
 	GShadowmapTotalStreamingSize += (ShadowmapFlags & SMF_Streamed) ? TextureSize : 0;
@@ -509,13 +515,23 @@ void FShadowMap2D::Serialize(FArchive& Ar)
 	{
 		Ar << bChannelValid[Channel];
 	}
+
+	if (Ar.UE4Ver() >= VER_UE4_STATIC_SHADOWMAP_PENUMBRA_SIZE)
+	{
+		Ar << InvUniformPenumbraSize;
+	}
+	else if (Ar.IsLoading())
+	{
+		const float LegacyValue = 1.0f / .05f;
+		InvUniformPenumbraSize = FVector4(LegacyValue, LegacyValue, LegacyValue, LegacyValue);
+	}
 }
 
 FShadowMapInteraction FShadowMap2D::GetInteraction() const
 {
 	if (Texture)
 	{
-		return FShadowMapInteraction::Texture(Texture, CoordinateScale, CoordinateBias, bChannelValid);
+		return FShadowMapInteraction::Texture(Texture, CoordinateScale, CoordinateBias, bChannelValid, InvUniformPenumbraSize);
 	}
 	return FShadowMapInteraction::None();
 }
@@ -754,18 +770,20 @@ void FShadowMap2D::EncodeTextures(UWorld* InWorld , bool bLightingSuccessful)
 	}
 }
 
-void FShadowMap2D::EncodeSingleTexture(FShadowMapPendingTexture& PendingTexture, UShadowMapTexture2D* Texture, TArray< TArray<FFourDistanceFieldSamples> >& MipData)
+int32 FShadowMap2D::EncodeSingleTexture(FShadowMapPendingTexture& PendingTexture, UShadowMapTexture2D* Texture, TArray< TArray<FFourDistanceFieldSamples> >& MipData)
 {
 	TArray<FFourDistanceFieldSamples>* TopMipData = new(MipData) TArray<FFourDistanceFieldSamples>();
 	TopMipData->Empty(PendingTexture.GetSizeX() * PendingTexture.GetSizeY());
 	TopMipData->AddZeroed(PendingTexture.GetSizeX() * PendingTexture.GetSizeY());
-	int32 TextureSizeX = Texture->Source.GetSizeX();
-	int32 TextureSizeY = Texture->Source.GetSizeY();
+	int32 TextureSizeX = PendingTexture.GetSizeX();
+	int32 TextureSizeY = PendingTexture.GetSizeY();
+	int32 MaxChannelsUsed = 0;
 
 	for (int32 AllocationIndex = 0;AllocationIndex < PendingTexture.Allocations.Num();AllocationIndex++)
 	{
 		FShadowMapAllocation& Allocation = *PendingTexture.Allocations[AllocationIndex];
 		bool bChannelUsed[4] = {0};
+		FVector4 InvUniformPenumbraSize(0, 0, 0, 0);
 
 		for (int32 ChannelIndex = 0; ChannelIndex < 4; ChannelIndex++)
 		{
@@ -773,9 +791,13 @@ void FShadowMap2D::EncodeSingleTexture(FShadowMapPendingTexture& PendingTexture,
 			{
 				if (ShadowMapPair.Key->ShadowMapChannel == ChannelIndex)
 				{
+					MaxChannelsUsed = FMath::Max(MaxChannelsUsed, ChannelIndex + 1);
 					bChannelUsed[ChannelIndex] = true;
 					const TArray<FQuantizedSignedDistanceFieldShadowSample>& SourceSamples = ShadowMapPair.Value;
 
+					// Warning - storing one penumbra size for the whole shadowmap even though multiple lights can share a channel
+					InvUniformPenumbraSize[ChannelIndex] = 1.0f / ShadowMapPair.Key->GetUniformPenumbraSize();
+
 					// Copy the raw data for this light-map into the raw texture data array.
 					for (int32 Y = Allocation.MappedRect.Min.Y; Y < Allocation.MappedRect.Max.Y; ++Y)
 					{
@@ -789,6 +811,7 @@ void FShadowMap2D::EncodeSingleTexture(FShadowMapPendingTexture& PendingTexture,
 
 							if ( SourceSample.Coverage > 0 )
 							{
+								// Note: multiple lights can write to different parts of the destination due to channel assignment
 								DestSample.Samples[ChannelIndex] = SourceSample;
 							}
 #if WITH_EDITOR
@@ -847,6 +870,8 @@ void FShadowMap2D::EncodeSingleTexture(FShadowMapPendingTexture& PendingTexture,
 		{
 			Allocation.ShadowMap->bChannelValid[ChannelIndex] = bChannelUsed[ChannelIndex];
 		}
+
+		Allocation.ShadowMap->InvUniformPenumbraSize = InvUniformPenumbraSize;
 	}
 
 	const uint32 NumMips = FMath::Max(FMath::CeilLogTwo(TextureSizeX),FMath::CeilLogTwo(TextureSizeY)) + 1;
@@ -1023,6 +1048,8 @@ void FShadowMap2D::EncodeSingleTexture(FShadowMapPendingTexture& PendingTexture,
 	{
 		Allocation->PostEncode();
 	}
+
+	return MaxChannelsUsed;
 }
 
 #endif
diff --git a/Engine/Source/Runtime/Engine/Public/SceneManagement.h b/Engine/Source/Runtime/Engine/Public/SceneManagement.h
--- a/Engine/Source/Runtime/Engine/Public/SceneManagement.h
+++ b/Engine/Source/Runtime/Engine/Public/SceneManagement.h
@@ -422,14 +422,16 @@ public:
 		class UShadowMapTexture2D* InTexture,
 		const FVector2D& InCoordinateScale,
 		const FVector2D& InCoordinateBias,
-		const bool* InChannelValid)
+		const bool* InChannelValid,
+		const FVector4& InInvUniformPenumbraSize)
 	{
 		FShadowMapInteraction Result;
 		Result.Type = SMIT_Texture;
 		Result.ShadowTexture = InTexture;
 		Result.CoordinateScale = InCoordinateScale;
 		Result.CoordinateBias = InCoordinateBias;
-		
+		Result.InvUniformPenumbraSize = InInvUniformPenumbraSize;
+
 		for (int Channel = 0; Channel < 4; Channel++)
 		{
 			Result.bChannelValid[Channel] = InChannelValid[Channel];
@@ -441,6 +443,7 @@ public:
 	/** Default constructor. */
 	FShadowMapInteraction() :
 		ShadowTexture(nullptr),
+		InvUniformPenumbraSize(FVector4(0, 0, 0, 0)),
 		Type(SMIT_None)
 	{
 		for (int Channel = 0; Channel < ARRAY_COUNT(bChannelValid); Channel++)
@@ -476,11 +479,17 @@ public:
 		return bChannelValid[ChannelIndex];
 	}
 
+	inline FVector4 GetInvUniformPenumbraSize() const
+	{
+		return InvUniformPenumbraSize;
+	}
+
 private:
 	UShadowMapTexture2D* ShadowTexture;
 	FVector2D CoordinateScale;
 	FVector2D CoordinateBias;
 	bool bChannelValid[4];
+	FVector4 InvUniformPenumbraSize;
 	EShadowMapInteractionType Type;
 };
 
diff --git a/Engine/Source/Runtime/Engine/Public/ShadowMap.h b/Engine/Source/Runtime/Engine/Public/ShadowMap.h
--- a/Engine/Source/Runtime/Engine/Public/ShadowMap.h
+++ b/Engine/Source/Runtime/Engine/Public/ShadowMap.h
@@ -287,7 +287,7 @@ public:
 	/**
 	 * Constructs mip maps for a single shadowmap texture.
 	 */
-	static void EncodeSingleTexture(struct FShadowMapPendingTexture& PendingTexture, UShadowMapTexture2D* Texture, TArray< TArray<FFourDistanceFieldSamples>>& MipData);
+	static int32 EncodeSingleTexture(struct FShadowMapPendingTexture& PendingTexture, UShadowMapTexture2D* Texture, TArray< TArray<FFourDistanceFieldSamples>>& MipData);
 
 	static ENGINE_API FShadowMap2D* AllocateShadowMap(
 		UObject* LightMapOuter,
@@ -354,6 +354,9 @@ private:
 	/** Tracks which of the 4 channels has valid texture data. */
 	bool bChannelValid[4];
 
+	/** Stores the inverse of the penumbra size, normalized.  Stores 1 to interpret the shadowmap as a shadow factor directly, instead of as a distance field. */
+	FVector4 InvUniformPenumbraSize;
+
 	/** If true, update the status when encoding light maps */
 	ENGINE_API static bool bUpdateStatus;
 };
diff --git a/Engine/Source/Runtime/Renderer/Private/LightMapRendering.h b/Engine/Source/Runtime/Renderer/Private/LightMapRendering.h
--- a/Engine/Source/Runtime/Renderer/Private/LightMapRendering.h
+++ b/Engine/Source/Runtime/Renderer/Private/LightMapRendering.h
@@ -300,22 +300,12 @@ public:
 		SuperElementDataType SuperElementData;
 		
 		FShadowMapInteraction ShadowMapInteraction;
-		FVector4 DistanceFieldValues;
 
 		/** Initialization constructor. */
 		ElementDataType( const FShadowMapInteraction& InShadowMapInteraction, const SuperElementDataType& InSuperElementData )
 			: SuperElementData(InSuperElementData)
 		{
 			ShadowMapInteraction = InShadowMapInteraction;
-
-			static const auto CVar0 = IConsoleManager::Get().FindTConsoleVariableDataFloat(TEXT("r.Shadow.DistanceFieldPenumbraSize"));
-			const float PenumbraSize = CVar0->GetValueOnAnyThread();
-
-			// Bias to convert distance from the distance field into the shadow penumbra based on penumbra size
-			DistanceFieldValues.X = 1.0f / FMath::Min(PenumbraSize, 1.0f);
-			DistanceFieldValues.Y = -0.5f * DistanceFieldValues.X + 0.5f;
-			DistanceFieldValues.Z = 0.0f;
-			DistanceFieldValues.W = 0.0f;
 		}
 	};
 
@@ -349,7 +339,7 @@ public:
 	public:
 		void Bind(const FShaderParameterMap& ParameterMap)
 		{
-			DistanceFieldParameters.Bind(ParameterMap,TEXT("DistanceFieldParameters"));
+			InvUniformPenumbraSizes.Bind(ParameterMap,TEXT("InvUniformPenumbraSizes"));
 			StaticShadowMapMasks.Bind(ParameterMap,TEXT("StaticShadowMapMasks"));
 			StaticShadowTexture.Bind(ParameterMap,TEXT("StaticShadowTexture"));
 			StaticShadowSampler.Bind(ParameterMap, TEXT("StaticShadowTextureSampler"));
@@ -357,12 +347,12 @@ public:
 			Super::PixelParametersType::Bind(ParameterMap);
 		}
 
-		void SetMesh(FRHICommandList& RHICmdList, FShader* PixelShader, const FShadowMapInteraction& ShadowMapInteraction, const FVector4& DistanceFieldValues) const
+		void SetMesh(FRHICommandList& RHICmdList, FShader* PixelShader, const FShadowMapInteraction& ShadowMapInteraction) const
 		{
 			const FPixelShaderRHIParamRef ShaderRHI = PixelShader->GetPixelShader();
 
-			SetShaderValue(RHICmdList, ShaderRHI, DistanceFieldParameters, DistanceFieldValues);
-			
+			SetShaderValue(RHICmdList, ShaderRHI, InvUniformPenumbraSizes, ShadowMapInteraction.GetInvUniformPenumbraSize());
+
 			SetShaderValue(RHICmdList, ShaderRHI, StaticShadowMapMasks, FVector4(
 				ShadowMapInteraction.GetChannelValid(0),
 				ShadowMapInteraction.GetChannelValid(1),
@@ -381,7 +371,7 @@ public:
 
 		void Serialize(FArchive& Ar)
 		{
-			Ar << DistanceFieldParameters;
+			Ar << InvUniformPenumbraSizes;
 			Ar << StaticShadowMapMasks;
 			Ar << StaticShadowTexture;
 			Ar << StaticShadowSampler;
@@ -392,7 +382,7 @@ public:
 		FShaderParameter StaticShadowMapMasks;
 		FShaderResourceParameter StaticShadowTexture;
 		FShaderResourceParameter StaticShadowSampler;
-		FShaderParameter DistanceFieldParameters;
+		FShaderParameter InvUniformPenumbraSizes;
 	};
 
 	static void ModifyCompilationEnvironment(EShaderPlatform Platform, const FMaterial* Material, FShaderCompilerEnvironment& OutEnvironment)
@@ -424,7 +414,7 @@ public:
 
 		if (PixelShaderParameters)
 		{
-			PixelShaderParameters->SetMesh(RHICmdList, PixelShader, ElementData.ShadowMapInteraction, ElementData.DistanceFieldValues);
+			PixelShaderParameters->SetMesh(RHICmdList, PixelShader, ElementData.ShadowMapInteraction);
 		}
 
 		Super::SetMesh(RHICmdList, View, PrimitiveSceneProxy, VertexShaderParameters, PixelShaderParameters, VertexShader, PixelShader, VertexFactory, MaterialRenderProxy, ElementData.SuperElementData);
