Build - Build fixes relating to the removal of the implicit constructor for FLinearColor to FColor.

diff --git a/Engine/Source/Programs/UnrealLightmass/Private/ImportExport/Exporter.cpp b/Engine/Source/Programs/UnrealLightmass/Private/ImportExport/Exporter.cpp
--- a/Engine/Source/Programs/UnrealLightmass/Private/ImportExport/Exporter.cpp
+++ b/Engine/Source/Programs/UnrealLightmass/Private/ImportExport/Exporter.cpp
@@ -135,7 +135,7 @@ namespace Lightmass
 				FLinearColor LightIntensity = CurrentLight.TotalPower / CurrentLight.TotalSurfaceArea;
 				// Extract an LDR light color and brightness scale
 				float MaxComponent = FMath::Max(LightIntensity.R, FMath::Max(LightIntensity.G, LightIntensity.B));
-				LightData.Color = LightIntensity / FMath::Max(MaxComponent, (float)KINDA_SMALL_NUMBER);
+				LightData.Color = ( LightIntensity / FMath::Max(MaxComponent, (float)KINDA_SMALL_NUMBER) ).ToFColor(true);
 				LightData.Brightness = MaxComponent;
 				LightData.FalloffExponent = CurrentLight.FalloffExponent;
 				Swarm->Write(&LightData, sizeof(LightData));
