Fixed: Precomputed Visibility no longer calculated on Landscapes UE-17594
Since landscape materials have to be exported to lightmass as masked, we compute visibility for landscape meshes regardless of material blend mode

diff --git a/Engine/Source/Programs/UnrealLightmass/Private/Lighting/Landscape.h b/Engine/Source/Programs/UnrealLightmass/Private/Lighting/Landscape.h
--- a/Engine/Source/Programs/UnrealLightmass/Private/Lighting/Landscape.h
+++ b/Engine/Source/Programs/UnrealLightmass/Private/Lighting/Landscape.h
@@ -23,6 +23,9 @@ namespace Lightmass
 		FORCEINLINE const FColor*	GetHeightData( int32 LocalX, int32 LocalY ) const;
 		FORCEINLINE void GetTriangleIndices(int32 QuadIndex,int32 TriNum,int32& OutI0,int32& OutI1,int32& OutI2) const;
 
+		// We always want to compute visibility for landscape meshes, regardless of material blend mode
+		virtual bool IsAlwaysOpaqueForVisibility() const override { return true; }
+
 	private:
 		TArray<FColor> HeightMap;
 		// Cache
diff --git a/Engine/Source/Programs/UnrealLightmass/Private/Lighting/LightingMesh.h b/Engine/Source/Programs/UnrealLightmass/Private/Lighting/LightingMesh.h
--- a/Engine/Source/Programs/UnrealLightmass/Private/Lighting/LightingMesh.h
+++ b/Engine/Source/Programs/UnrealLightmass/Private/Lighting/LightingMesh.h
@@ -386,6 +386,12 @@ public:
 	/** For debugging */
 	virtual void SetDebugMaterial(bool bUseDebugMaterial, FLinearColor Diffuse);
 
+	/** 
+	 * Whether mesh is always opaque for visibility calculations, 
+	 * otherwise opaque property will be checked for each triangle 
+	 */
+	virtual bool IsAlwaysOpaqueForVisibility() const { return false; }
+
 	/** Evaluates the mesh's Bidirectional Reflectance Distribution Function. */
 	FLinearColor EvaluateBRDF(
 		const FStaticLightingVertex& Vertex, 
diff --git a/Engine/Source/Programs/UnrealLightmass/Private/Lighting/PrecomputedVisibility.cpp b/Engine/Source/Programs/UnrealLightmass/Private/Lighting/PrecomputedVisibility.cpp
--- a/Engine/Source/Programs/UnrealLightmass/Private/Lighting/PrecomputedVisibility.cpp
+++ b/Engine/Source/Programs/UnrealLightmass/Private/Lighting/PrecomputedVisibility.cpp
@@ -219,6 +219,9 @@ void FStaticLightingSystem::SetupPrecomputedVisibility()
 			// Only process meshes whose bounding box intersects a PVS volume
 			if (Scene.DoesBoxIntersectVisibilityVolume(CurrentMesh->BoundingBox))
 			{
+				// Whether mesh wants to be fully opaque for visibility step
+				const bool bOpaqueMesh = CurrentMesh->IsAlwaysOpaqueForVisibility();
+				
 				// Rasterize all triangles in the mesh
 				for (int32 TriangleIndex = 0; TriangleIndex < CurrentMesh->NumTriangles; TriangleIndex++)
 				{
@@ -227,7 +230,7 @@ void FStaticLightingSystem::SetupPrecomputedVisibility()
 					CurrentMesh->GetTriangle(TriangleIndex, Vertices[0], Vertices[1], Vertices[2], ElementIndex);
 
 					// Only place cells on opaque surfaces if requested, which can save some memory for foliage maps
-					if (!PrecomputedVisibilitySettings.bPlaceCellsOnOpaqueOnly 
+					if (!PrecomputedVisibilitySettings.bPlaceCellsOnOpaqueOnly || bOpaqueMesh
 						|| (!CurrentMesh->IsMasked(ElementIndex) && !CurrentMesh->IsTranslucent(ElementIndex)))
 					{
 						FVector2D XYPositions[3];
