Volume lighting samples are no longer placed inside geometry
Volume lighting sample computation is now multithreaded in Lightmass so it can't bottleneck the lighting build

diff --git a/Engine/Config/BaseLightmass.ini b/Engine/Config/BaseLightmass.ini
--- a/Engine/Config/BaseLightmass.ini
+++ b/Engine/Config/BaseLightmass.ini
@@ -65,8 +65,8 @@ SurfaceSampleLayerHeightSpacing=250
 NumSurfaceSampleLayers=2
 DetailVolumeSampleSpacing=300
 VolumeLightSampleSpacing=3000
-; Clamp the number of volume samples generated to ~3mb
-MaxVolumeSamples=50000
+; Clamp the number of volume samples generated to ~15mb
+MaxVolumeSamples=250000
 bUseMaxSurfaceSampleNum=True
 ; Approximately clamp the number of surface samples generated to ~30mb (only for Landscape currently)
 MaxSurfaceLightSamples=500000
diff --git a/Engine/Source/Programs/UnrealLightmass/Private/Lighting/LightingSystem.cpp b/Engine/Source/Programs/UnrealLightmass/Private/Lighting/LightingSystem.cpp
--- a/Engine/Source/Programs/UnrealLightmass/Private/Lighting/LightingSystem.cpp
+++ b/Engine/Source/Programs/UnrealLightmass/Private/Lighting/LightingSystem.cpp
@@ -208,7 +208,9 @@ FStaticLightingSystem::FStaticLightingSystem(const FLightingBuildOptions& InOpti
 ,	PhotonMappingSettings(InScene.PhotonMappingSettings)
 ,	IrradianceCachingSettings(InScene.IrradianceCachingSettings)
 ,	MappingTasksInProgressThatWillNeedHelp(0)
-,	bVolumeLightingSamplesComplete(false)
+,	NextVolumeSampleTaskIndex(-1)
+,	NumVolumeSampleTasksOutstanding(0)
+,	bShouldExportVolumeSampleData(false)
 ,	VolumeLightingInterpolationOctree(FVector4(0,0,0), HALF_WORLD_MAX)
 ,	bShouldExportMeshAreaLightData(false)
 ,	bShouldExportVolumeDistanceField(false)
@@ -538,7 +540,7 @@ void FStaticLightingSystem::MultithreadProcess()
 	{
 		// Calculate volume samples now if they will be needed by the lighting threads for shading,
 		// Otherwise the volume samples will be calculated when the task is received from swarm.
-		CalculateVolumeSamples();
+		BeginCalculateVolumeSamples();
 	}
 
 	SetupPrecomputedVisibility();
@@ -645,8 +647,10 @@ void FStaticLightingSystem::MultithreadProcess()
 void FStaticLightingSystem::ExportNonMappingTasks()
 {
 	// Export volume lighting samples to Swarm if they are complete
-	if (bVolumeLightingSamplesComplete)
+	if (bShouldExportVolumeSampleData)
 	{
+		bShouldExportVolumeSampleData = false;
+
 		Exporter.ExportVolumeLightingSamples(
 			DynamicObjectSettings.bVisualizeVolumeLightSamples,
 			VolumeLightingDebugOutput,
@@ -666,7 +670,6 @@ void FStaticLightingSystem::ExportNonMappingTasks()
 			FLightmassSwarm* Swarm = GetExporter().GetSwarm();
 			Swarm->TaskCompleted( PrecomputedVolumeLightingGuid );
 		}
-		bVolumeLightingSamplesComplete = 0;
 	}
 
 	CompleteVisibilityTaskList.ApplyAndClear(*this);
@@ -908,7 +911,7 @@ void FStaticLightingSystem::DumpStats(float TotalStaticLightingTime) const
 		}
 	}
 
-	if (Stats.PrecomputedVisibilitySetupTime > 0)
+	if (Stats.PrecomputedVisibilitySetupTime / TotalStaticLightingTime > .02f)
 	{
 		SolverStats += FString::Printf( TEXT("%4.1f%%%8.1fs    sPVS setup\n"), 100.0f * Stats.PrecomputedVisibilitySetupTime / TotalStaticLightingTime, Stats.PrecomputedVisibilitySetupTime);
 	}
@@ -1026,9 +1029,9 @@ void FStaticLightingSystem::DumpStats(float TotalStaticLightingTime) const
 	}
 	if (Stats.NumDynamicObjectSurfaceSamples + Stats.NumDynamicObjectVolumeSamples > 0)
 	{
-		SolverStats += FString::Printf( TEXT("%4.1f%%%8.1fs    Volume Samples\n"), 100.0f * Stats.VolumeSampleThreadTime / TotalLightingBusyThreadTime, Stats.VolumeSampleThreadTime);
+		SolverStats += FString::Printf( TEXT("%4.1f%%%8.1fs    Volume Sample placement\n"), 100.0f * Stats.VolumeSamplePlacementThreadTime / TotalLightingBusyThreadTime, Stats.VolumeSamplePlacementThreadTime);
 	}
-	const float UnaccountedLightingThreadTime = FMath::Max(TotalLightingBusyThreadTime - (SampleSetupTime + Stats.DirectLightingTime + Stats.BlockOnIndirectLightingCacheTasksTime + Stats.BlockOnIndirectLightingInterpolateTasksTime + Stats.SecondPassIrradianceCacheInterpolationTime + Stats.VolumeSampleThreadTime + Stats.StaticShadowDepthMapThreadTime + Stats.VolumeDistanceFieldThreadTime + PrecomputedVisibilityThreadTime), 0.0f);
+	const float UnaccountedLightingThreadTime = FMath::Max(TotalLightingBusyThreadTime - (SampleSetupTime + Stats.DirectLightingTime + Stats.BlockOnIndirectLightingCacheTasksTime + Stats.BlockOnIndirectLightingInterpolateTasksTime + Stats.SecondPassIrradianceCacheInterpolationTime + Stats.VolumeSamplePlacementThreadTime + Stats.StaticShadowDepthMapThreadTime + Stats.VolumeDistanceFieldThreadTime + PrecomputedVisibilityThreadTime), 0.0f);
 	SolverStats += FString::Printf( TEXT("%4.1f%%%8.1fs    Unaccounted\n"), 100.0f * UnaccountedLightingThreadTime / TotalLightingBusyThreadTime, UnaccountedLightingThreadTime);
 	// Send the message in multiple parts since it cuts off in the middle otherwise
 	LogSolverMessage(SolverStats);
@@ -1073,7 +1076,7 @@ void FStaticLightingSystem::DumpStats(float TotalStaticLightingTime) const
 	const int32 TotalVolumeLightingSamples = Stats.NumDynamicObjectSurfaceSamples + Stats.NumDynamicObjectVolumeSamples;
 	if (TotalVolumeLightingSamples > 0)
 	{
-		SolverStats += FString::Printf( TEXT("%u Volume lighting samples, %.1f%% placed on surfaces, %.1f%% placed in the volume\n"), TotalVolumeLightingSamples, 100.0f * Stats.NumDynamicObjectSurfaceSamples / (float)TotalVolumeLightingSamples, 100.0f * Stats.NumDynamicObjectVolumeSamples / (float)TotalVolumeLightingSamples);
+		SolverStats += FString::Printf( TEXT("%u Volume lighting samples, %.1f%% placed on surfaces, %.1f%% placed in the volume, %.1f thread seconds\n"), TotalVolumeLightingSamples, 100.0f * Stats.NumDynamicObjectSurfaceSamples / (float)TotalVolumeLightingSamples, 100.0f * Stats.NumDynamicObjectVolumeSamples / (float)TotalVolumeLightingSamples, Stats.TotalVolumeSampleLightingThreadTime);
 	}
 
 	if (Stats.NumPrecomputedVisibilityQueries > 0)
@@ -1548,8 +1551,7 @@ void FStaticLightingSystem::ThreadLoop(bool bIsMainThread, int32 ThreadIndex, FT
 		}
 		else if (bDynamicObjectTask)
 		{
-			CalculateVolumeSamples();
-			FPlatformAtomics::InterlockedExchange(&bVolumeLightingSamplesComplete, true);
+			BeginCalculateVolumeSamples();
 		}
 		else if (PrecomputedVisibilityTaskIndex >= 0)
 		{
@@ -1593,9 +1595,26 @@ void FStaticLightingSystem::ThreadLoop(bool bIsMainThread, int32 ThreadIndex, FT
 				NextInterpolateTask->TextureMapping->CompletedInterpolationTasks.Push(NextInterpolateTask);
 				FPlatformAtomics::InterlockedDecrement(&NextInterpolateTask->TextureMapping->NumOutstandingInterpolationTasks);
 			}
+			
+			if (NumVolumeSampleTasksOutstanding > 0)
+			{
+				const int32 TaskIndex = FPlatformAtomics::InterlockedIncrement(&NextVolumeSampleTaskIndex);
+
+				if (TaskIndex < VolumeSampleTasks.Num())
+				{
+					ProcessVolumeSamplesTask(VolumeSampleTasks[TaskIndex]);
+					const int32 NumTasksRemaining = FPlatformAtomics::InterlockedDecrement(&NumVolumeSampleTasksOutstanding);
+
+					if (NumTasksRemaining == 0)
+					{
+						FPlatformAtomics::InterlockedExchange(&bShouldExportVolumeSampleData, true);
+					}
+				}
+			}
 
 			if (!NextCacheTask 
 				&& !NextInterpolateTask
+				&& NumVolumeSampleTasksOutstanding <= 0
 				&& NumOutstandingVolumeDataLayers <= 0)
 			{
 				if (MappingTasksInProgressThatWillNeedHelp <= 0 && !bRequestForTaskTimedOut)
diff --git a/Engine/Source/Programs/UnrealLightmass/Private/Lighting/LightingSystem.h b/Engine/Source/Programs/UnrealLightmass/Private/Lighting/LightingSystem.h
--- a/Engine/Source/Programs/UnrealLightmass/Private/Lighting/LightingSystem.h
+++ b/Engine/Source/Programs/UnrealLightmass/Private/Lighting/LightingSystem.h
@@ -1181,7 +1181,7 @@ public:
 	float BooleanRayTraceThreadTime;
 
 	/** Thread seconds spent placing volume lighting samples and then calculating their lighting. */
-	float VolumeSampleThreadTime;
+	float VolumeSamplePlacementThreadTime;
 
 	/** Irradiance cache stats */
 	FIrradianceCacheStats Cache[NumTrackedBounces];
@@ -1244,7 +1244,7 @@ public:
 		NumBooleanRaysTraced(0),
 		FirstHitRayTraceThreadTime(0),
 		BooleanRayTraceThreadTime(0),
-		VolumeSampleThreadTime(0)
+		VolumeSamplePlacementThreadTime(0)
 	{
 		for (int32 i = 0; i < NumTrackedBounces; i++)
 		{
@@ -1891,6 +1891,21 @@ public:
 	{}
 };
 
+/**  */
+class FVolumeSamplesTaskDescription
+{
+public:
+	FGuid LevelId;
+	int32 StartIndex;
+	int32 NumSamples;
+	
+	FVolumeSamplesTaskDescription(FGuid InLevelId, int32 InStartIndex, INT32 InNumSamples) :
+		LevelId(InLevelId),
+		StartIndex(InStartIndex),
+		NumSamples(InNumSamples)
+	{}
+};
+
 class FVisibilityMeshGroup
 {
 public:
@@ -1973,6 +1988,11 @@ public:
 		FVector2D UVBias, 
 		FVector2D UVScale) const;
 
+	FStaticLightingAggregateMesh& GetAggregateMesh()
+	{
+		return AggregateMesh;
+	}
+
 private:
 
 	/** Exports tasks that are not mappings, if they are ready. */
@@ -2156,7 +2176,7 @@ private:
 		bool bDebugThisDensityEstimation) const;
 
 	/** Places volume lighting samples and calculates lighting for them. */
-	void CalculateVolumeSamples();
+	void BeginCalculateVolumeSamples();
 
 	/** 
 	 * Interpolates lighting from the volume lighting samples to a vertex. 
@@ -2469,6 +2489,8 @@ private:
 	 */
 	void ProcessInterpolateTask(FInterpolateIndirectTaskDescription* Task, bool bProcessedByMappingThread);
 
+	void ProcessVolumeSamplesTask(const FVolumeSamplesTaskDescription& Task);
+
 	/** Handles indirect lighting calculations for a single texture mapping. */
 	void CalculateIndirectLightingTextureMapping(
 		FStaticLightingTextureMapping* TextureMapping,
@@ -2572,8 +2594,12 @@ private:
 	/** List of tasks to interpolate indirect lighting, used by all mapping threads. */
 	TLockFreePointerList<FInterpolateIndirectTaskDescription> InterpolateIndirectLightingTasks;
 
-	/* Positive if the lighting threads are done writing to VolumeLightingSamples, and the samples can be exported. */
-	volatile int32 bVolumeLightingSamplesComplete;
+	TArray<FVolumeSamplesTaskDescription> VolumeSampleTasks;
+
+	volatile int32 bHasVolumeSampleTasks;
+	volatile int32 NextVolumeSampleTaskIndex;
+	volatile int32 NumVolumeSampleTasksOutstanding;
+	volatile int32 bShouldExportVolumeSampleData;
 	/** Bounds that VolumeLightingSamples were generated in. */
 	FBoxSphereBounds VolumeBounds;
 	/** Octree used for interpolating the volume lighting samples if DynamicObjectSettings.bVisualizeVolumeLightInterpolation is true. */
diff --git a/Engine/Source/Programs/UnrealLightmass/Private/Lighting/SampleVolume.cpp b/Engine/Source/Programs/UnrealLightmass/Private/Lighting/SampleVolume.cpp
--- a/Engine/Source/Programs/UnrealLightmass/Private/Lighting/SampleVolume.cpp
+++ b/Engine/Source/Programs/UnrealLightmass/Private/Lighting/SampleVolume.cpp
@@ -88,6 +88,7 @@ public:
 		int32 InSizeX, 
 		int32 InSizeY, 
 		float InMinSampleDistance, 
+		float InSceneBoundingRadius,
 		FStaticLightingSystem& InSystem,
 		FCoherentRayCache& InCoherentRayCache,
 		FVolumeLightingProximityOctree& InProximityOctree)
@@ -95,6 +96,7 @@ public:
 		SizeX(InSizeX),
 		SizeY(InSizeY),
 		MinSampleDistance(InMinSampleDistance),
+		SceneBoundingRadius(InSceneBoundingRadius),
 		System(InSystem),
 		CoherentRayCache(InCoherentRayCache),
 		ProximityOctree(InProximityOctree)
@@ -105,6 +107,16 @@ public:
 		{
 			LayerHeightOffsets.Add(System.DynamicObjectSettings.FirstSurfaceSampleLayerHeight + i * System.DynamicObjectSettings.SurfaceSampleLayerHeightSpacing);
 		}
+
+		
+		TArray<FVector2D> UniformHemisphereSampleUniforms;
+		const int32 NumUpperVolumeSamples = 16;
+		const float NumThetaStepsFloat = FMath::Sqrt(NumUpperVolumeSamples / (float)PI);
+		const int32 NumThetaSteps = FMath::TruncToInt(NumThetaStepsFloat);
+		const int32 NumPhiSteps = FMath::TruncToInt(NumThetaStepsFloat * (float)PI);
+		FLMRandomStream RandomStream(0);
+
+		GenerateStratifiedUniformHemisphereSamples(NumThetaSteps, NumPhiSteps, RandomStream, UniformHemisphereSamples, UniformHemisphereSampleUniforms);
 	}
 
 	void SetLevelGuid(FGuid InLevelGuid)
@@ -128,35 +140,76 @@ private:
 	const int32 SizeX;
 	const int32 SizeY;
 	const float MinSampleDistance;
+	const float SceneBoundingRadius;
 	FGuid LevelGuid;
 	FStaticLightingSystem& System;
 	FCoherentRayCache& CoherentRayCache;
 	FVolumeLightingProximityOctree& ProximityOctree;
 	TArray<float> LayerHeightOffsets;
+	TArray<FVector4> UniformHemisphereSamples;
 };
 
+int32 ComputeNumBackfacingHits(
+	const FVector4& SamplePosition, 
+	FStaticLightingSystem& System, 
+	FCoherentRayCache& CoherentRayCache, 
+	float SceneBoundingRadius, 
+	const TArray<FVector4>& UniformHemisphereSamples)
+{
+	int NumBackfacingHits = 0;
+
+	for (int32 SampleIndex = 0; SampleIndex < UniformHemisphereSamples.Num() * 2; SampleIndex++)
+	{
+		FVector4 SampleDirection = UniformHemisphereSamples[SampleIndex % UniformHemisphereSamples.Num()];
+		SampleDirection.Z *= SampleIndex >= UniformHemisphereSamples.Num() ? -1.0f : 1.0f;
+
+		const FLightRay PathRay(
+			SamplePosition,
+			SamplePosition + SampleDirection * SceneBoundingRadius,
+			NULL,
+			NULL
+			);
+
+		FLightRayIntersection RayIntersection;
+		System.GetAggregateMesh().IntersectLightRay(PathRay, true, false, false, CoherentRayCache, RayIntersection);
+
+		if (RayIntersection.bIntersects && Dot3(PathRay.Direction, -RayIntersection.IntersectionVertex.WorldTangentZ) <= 0.0f)
+		{
+			NumBackfacingHits++;
+		}
+	}
+
+	return NumBackfacingHits;
+}
+
 void FVolumeSamplePlacementRasterPolicy::ProcessPixel(int32 X,int32 Y,const InterpolantType& Vertex,bool BackFacing)
 {
 	// Only place samples inside the scene's bounds
 	if (System.IsPointInImportanceVolume(Vertex.WorldPosition))
 	{
 		// Place a sample for each layer
-		for (int32 SampleIndex = 0; SampleIndex < LayerHeightOffsets.Num(); SampleIndex++)
+		for (int32 LayerIndex = 0; LayerIndex < LayerHeightOffsets.Num(); LayerIndex++)
 		{
-			const FVector4 SamplePosition = Vertex.WorldPosition + FVector4(0, 0, LayerHeightOffsets[SampleIndex]);
+			const FVector4 SamplePosition = Vertex.WorldPosition + FVector4(0, 0, LayerHeightOffsets[LayerIndex]);
 			// Only place a sample if there isn't already one nearby, to avoid clumping
 			if (!FindNearbyVolumeSample(ProximityOctree, SamplePosition, MinSampleDistance))
 			{
-				TArray<FVolumeLightingSample>* VolumeLightingSamples = System.VolumeLightingSamples.Find(LevelGuid);
-				check(VolumeLightingSamples);
-				// Add a new sample for this layer
-				// Expand the radius to touch a diagonal sample on the grid for a little overlap
-				VolumeLightingSamples->Add(FVolumeLightingSample(FVector4(SamplePosition, System.DynamicObjectSettings.SurfaceLightSampleSpacing * FMath::Sqrt(2.0f))));
-				// Add the sample to the proximity octree so we can avoid placing any more samples nearby
-				ProximityOctree.AddElement(FVolumeSampleProximityElement(VolumeLightingSamples->Num() - 1, *VolumeLightingSamples));
-				if (System.DynamicObjectSettings.bVisualizeVolumeLightInterpolation)
+				const int32 NumBackfacingHits = ComputeNumBackfacingHits(SamplePosition, System, CoherentRayCache, SceneBoundingRadius, UniformHemisphereSamples);
+
+				// Only place a sample if we are outside of the level geometry (determined by whether we can see backfaces)
+				if (NumBackfacingHits < .3f * UniformHemisphereSamples.Num() * 2)
 				{
-					System.VolumeLightingInterpolationOctree.AddElement(FVolumeSampleInterpolationElement(VolumeLightingSamples->Num() - 1, *VolumeLightingSamples));
+					TArray<FVolumeLightingSample>* VolumeLightingSamples = System.VolumeLightingSamples.Find(LevelGuid);
+					check(VolumeLightingSamples);
+					// Add a new sample for this layer
+					// Expand the radius to touch a diagonal sample on the grid for a little overlap
+					VolumeLightingSamples->Add(FVolumeLightingSample(FVector4(SamplePosition, System.DynamicObjectSettings.SurfaceLightSampleSpacing * FMath::Sqrt(2.0f))));
+					// Add the sample to the proximity octree so we can avoid placing any more samples nearby
+					ProximityOctree.AddElement(FVolumeSampleProximityElement(VolumeLightingSamples->Num() - 1, *VolumeLightingSamples));
+					if (System.DynamicObjectSettings.bVisualizeVolumeLightInterpolation)
+					{
+						System.VolumeLightingInterpolationOctree.AddElement(FVolumeSampleInterpolationElement(VolumeLightingSamples->Num() - 1, *VolumeLightingSamples));
+					}
 				}
 			}
 		}
@@ -164,310 +217,328 @@ void FVolumeSamplePlacementRasterPolicy::ProcessPixel(int32 X,int32 Y,const Inte
 }
 
 /** Places volume lighting samples and calculates lighting for them. */
-void FStaticLightingSystem::CalculateVolumeSamples()
+void FStaticLightingSystem::BeginCalculateVolumeSamples()
 {
-	const double SampleVolumeStart = FPlatformTime::Seconds();
-	if (VolumeLightingSamples.Num() == 0)
+	const double VolumeSampleStartTime = FPlatformTime::Seconds();
+	VolumeBounds = GetImportanceBounds(false);
+	if (VolumeBounds.SphereRadius < DELTA)
 	{
-		const double VolumeSampleStartTime = FPlatformTime::Seconds();
-		VolumeBounds = GetImportanceBounds(false);
-		if (VolumeBounds.SphereRadius < DELTA)
-		{
-			VolumeBounds = FBoxSphereBounds(AggregateMesh.GetBounds());
-		}
+		VolumeBounds = FBoxSphereBounds(AggregateMesh.GetBounds());
+	}
 
-		// Only place samples if the volume has area
-		if (VolumeBounds.BoxExtent.X > 0.0f && VolumeBounds.BoxExtent.Y > 0.0f && VolumeBounds.BoxExtent.Z > 0.0f)
+	// Only place samples if the volume has area
+	if (VolumeBounds.BoxExtent.X > 0.0f && VolumeBounds.BoxExtent.Y > 0.0f && VolumeBounds.BoxExtent.Z > 0.0f)
+	{
+		float LandscapeEstimateNum = 0.f;
+		// Estimate Light sample number near Landscape surfaces
+		if (DynamicObjectSettings.bUseMaxSurfaceSampleNum && DynamicObjectSettings.MaxSurfaceLightSamples > 100)
 		{
-			float LandscapeEstimateNum = 0.f;
-			// Estimate Light sample number near Landscape surfaces
-			if (DynamicObjectSettings.bUseMaxSurfaceSampleNum && DynamicObjectSettings.MaxSurfaceLightSamples > 100)
+			float SquaredSpacing = FMath::Square(DynamicObjectSettings.SurfaceLightSampleSpacing);
+			if (SquaredSpacing == 0.f) SquaredSpacing = 1.0f;
+			for (int32 MappingIndex = 0; MappingIndex < LandscapeMappings.Num(); MappingIndex++)
 			{
-				float SquaredSpacing = FMath::Square(DynamicObjectSettings.SurfaceLightSampleSpacing);
-				if (SquaredSpacing == 0.f) SquaredSpacing = 1.0f;
-				for (int32 MappingIndex = 0; MappingIndex < LandscapeMappings.Num(); MappingIndex++)
+				FStaticLightingVertex Vertices[3];
+				int32 ElementIndex;
+				const FStaticLightingMapping* CurrentMapping = LandscapeMappings[MappingIndex];
+				const FStaticLightingMesh* CurrentMesh = CurrentMapping->Mesh;
+				CurrentMesh->GetTriangle((CurrentMesh->NumTriangles)>>1, Vertices[0], Vertices[1], Vertices[2], ElementIndex);
+				// Only place inside the importance volume
+				if (IsPointInImportanceVolume(Vertices[0].WorldPosition))
 				{
-					FStaticLightingVertex Vertices[3];
-					int32 ElementIndex;
-					const FStaticLightingMapping* CurrentMapping = LandscapeMappings[MappingIndex];
-					const FStaticLightingMesh* CurrentMesh = CurrentMapping->Mesh;
-					CurrentMesh->GetTriangle((CurrentMesh->NumTriangles)>>1, Vertices[0], Vertices[1], Vertices[2], ElementIndex);
-					// Only place inside the importance volume
-					if (IsPointInImportanceVolume(Vertices[0].WorldPosition))
-					{
-						FVector4 TriangleNormal = (Vertices[2].WorldPosition - Vertices[0].WorldPosition) ^ (Vertices[1].WorldPosition - Vertices[0].WorldPosition);
-						TriangleNormal.Z = 0.f; // approximate only for X-Y plane
-						float TotalArea = 0.5f * TriangleNormal.Size3() * CurrentMesh->NumTriangles;
-						LandscapeEstimateNum += TotalArea / FMath::Square(DynamicObjectSettings.SurfaceLightSampleSpacing);
-					}
+					FVector4 TriangleNormal = (Vertices[2].WorldPosition - Vertices[0].WorldPosition) ^ (Vertices[1].WorldPosition - Vertices[0].WorldPosition);
+					TriangleNormal.Z = 0.f; // approximate only for X-Y plane
+					float TotalArea = 0.5f * TriangleNormal.Size3() * CurrentMesh->NumTriangles;
+					LandscapeEstimateNum += TotalArea / FMath::Square(DynamicObjectSettings.SurfaceLightSampleSpacing);
 				}
-				LandscapeEstimateNum *= DynamicObjectSettings.NumSurfaceSampleLayers;
+			}
+			LandscapeEstimateNum *= DynamicObjectSettings.NumSurfaceSampleLayers;
 
-				if (LandscapeEstimateNum > DynamicObjectSettings.MaxSurfaceLightSamples)
-				{
-					// Increase DynamicObjectSettings.SurfaceLightSampleSpacing to reduce light sample number
-					float OldMaxSurfaceLightSamples = DynamicObjectSettings.SurfaceLightSampleSpacing;
-					DynamicObjectSettings.SurfaceLightSampleSpacing = DynamicObjectSettings.SurfaceLightSampleSpacing * FMath::Sqrt((float)LandscapeEstimateNum / DynamicObjectSettings.MaxSurfaceLightSamples);
-					UE_LOG(LogLightmass, Log, TEXT("Too many LightSamples : DynamicObjectSettings.SurfaceLightSampleSpacing is increased from %g to %g"), OldMaxSurfaceLightSamples, DynamicObjectSettings.SurfaceLightSampleSpacing);
-					LandscapeEstimateNum = DynamicObjectSettings.MaxSurfaceLightSamples;
-				}
+			if (LandscapeEstimateNum > DynamicObjectSettings.MaxSurfaceLightSamples)
+			{
+				// Increase DynamicObjectSettings.SurfaceLightSampleSpacing to reduce light sample number
+				float OldMaxSurfaceLightSamples = DynamicObjectSettings.SurfaceLightSampleSpacing;
+				DynamicObjectSettings.SurfaceLightSampleSpacing = DynamicObjectSettings.SurfaceLightSampleSpacing * FMath::Sqrt((float)LandscapeEstimateNum / DynamicObjectSettings.MaxSurfaceLightSamples);
+				UE_LOG(LogLightmass, Log, TEXT("Too many LightSamples : DynamicObjectSettings.SurfaceLightSampleSpacing is increased from %g to %g"), OldMaxSurfaceLightSamples, DynamicObjectSettings.SurfaceLightSampleSpacing);
+				LandscapeEstimateNum = DynamicObjectSettings.MaxSurfaceLightSamples;
 			}
+		}
 
-			//@todo - can this be presized more accurately?
-			VolumeLightingSamples.Empty(FMath::Max<int32>(5000, LandscapeEstimateNum));
-			FStaticLightingMappingContext MappingContext(nullptr, *this);
-			// Octree used to keep track of where existing samples have been placed
-			FVolumeLightingProximityOctree VolumeLightingOctree(VolumeBounds.Origin, VolumeBounds.BoxExtent.GetMax());
-			// Octree used for interpolating lighting for debugging
-			VolumeLightingInterpolationOctree = FVolumeLightingInterpolationOctree(VolumeBounds.Origin, VolumeBounds.BoxExtent.GetMax());
-			// Determine the resolution that the scene should be rasterized at based on SurfaceLightSampleSpacing and the scene's extent
-			const int32 RasterSizeX = FMath::TruncToInt(2.0f * VolumeBounds.BoxExtent.X / DynamicObjectSettings.SurfaceLightSampleSpacing);
-			const int32 RasterSizeY = FMath::TruncToInt(2.0f * VolumeBounds.BoxExtent.Y / DynamicObjectSettings.SurfaceLightSampleSpacing);
-
-			FTriangleRasterizer<FVolumeSamplePlacementRasterPolicy> Rasterizer(
-				FVolumeSamplePlacementRasterPolicy(
-				RasterSizeX, 
-				RasterSizeY, 
-				// Use a minimum sample distance slightly less than the SurfaceLightSampleSpacing
-				0.9f * FMath::Min(DynamicObjectSettings.SurfaceLightSampleSpacing, DynamicObjectSettings.SurfaceSampleLayerHeightSpacing), 
-				*this,
-				MappingContext.RayCache,
-				VolumeLightingOctree));
-
-			check(Meshes.Num() == AllMappings.Num());
-			// Rasterize all meshes in the scene and place high detail samples on their surfaces.
-			// Iterate through mappings and retreive the mesh from that, so we can make decisions based on whether the mesh is using texture or vertex lightmaps.
-			for (int32 MappingIndex = 0; MappingIndex < AllMappings.Num(); MappingIndex++)
+		//@todo - can this be presized more accurately?
+		VolumeLightingSamples.Empty(FMath::Max<int32>(5000, LandscapeEstimateNum));
+		FStaticLightingMappingContext MappingContext(nullptr, *this);
+		// Octree used to keep track of where existing samples have been placed
+		FVolumeLightingProximityOctree VolumeLightingOctree(VolumeBounds.Origin, VolumeBounds.BoxExtent.GetMax());
+		// Octree used for interpolating lighting for debugging
+		VolumeLightingInterpolationOctree = FVolumeLightingInterpolationOctree(VolumeBounds.Origin, VolumeBounds.BoxExtent.GetMax());
+		// Determine the resolution that the scene should be rasterized at based on SurfaceLightSampleSpacing and the scene's extent
+		const int32 RasterSizeX = FMath::TruncToInt(2.0f * VolumeBounds.BoxExtent.X / DynamicObjectSettings.SurfaceLightSampleSpacing);
+		const int32 RasterSizeY = FMath::TruncToInt(2.0f * VolumeBounds.BoxExtent.Y / DynamicObjectSettings.SurfaceLightSampleSpacing);
+
+		FTriangleRasterizer<FVolumeSamplePlacementRasterPolicy> Rasterizer(
+			FVolumeSamplePlacementRasterPolicy(
+			RasterSizeX, 
+			RasterSizeY, 
+			// Use a minimum sample distance slightly less than the SurfaceLightSampleSpacing
+			0.9f * FMath::Min(DynamicObjectSettings.SurfaceLightSampleSpacing, DynamicObjectSettings.SurfaceSampleLayerHeightSpacing), 
+			FBoxSphereBounds(AggregateMesh.GetBounds()).SphereRadius,
+			*this,
+			MappingContext.RayCache,
+			VolumeLightingOctree));
+
+		check(Meshes.Num() == AllMappings.Num());
+		// Rasterize all meshes in the scene and place high detail samples on their surfaces.
+		// Iterate through mappings and retreive the mesh from that, so we can make decisions based on whether the mesh is using texture or vertex lightmaps.
+		for (int32 MappingIndex = 0; MappingIndex < AllMappings.Num(); MappingIndex++)
+		{
+			const FStaticLightingMapping* CurrentMapping = AllMappings[MappingIndex];
+			const FStaticLightingTextureMapping* TextureMapping = CurrentMapping->GetTextureMapping();
+			const FStaticLightingMesh* CurrentMesh = CurrentMapping->Mesh;
+			// Only place samples on shadow casting meshes.
+			if (CurrentMesh->LightingFlags & GI_INSTANCE_CASTSHADOW)
 			{
-				const FStaticLightingMapping* CurrentMapping = AllMappings[MappingIndex];
-				const FStaticLightingTextureMapping* TextureMapping = CurrentMapping->GetTextureMapping();
-				const FStaticLightingMesh* CurrentMesh = CurrentMapping->Mesh;
-				// Only place samples on shadow casting meshes.
-				if (CurrentMesh->LightingFlags & GI_INSTANCE_CASTSHADOW)
+				// Create a new LevelId array if necessary
+				if (!VolumeLightingSamples.Find(CurrentMesh->LevelGuid))
 				{
-					// Create a new LevelId array if necessary
-					if (!VolumeLightingSamples.Find(CurrentMesh->LevelGuid))
-					{
-						VolumeLightingSamples.Add(CurrentMesh->LevelGuid, TArray<FVolumeLightingSample>());
-					}
-					// Tell the rasterizer we are adding samples to this mesh's LevelId
-					Rasterizer.SetLevelGuid(CurrentMesh->LevelGuid);
-					// Rasterize all triangles in the mesh
-					for (int32 TriangleIndex = 0; TriangleIndex < CurrentMesh->NumTriangles; TriangleIndex++)
-					{
-						FStaticLightingVertex Vertices[3];
-						int32 ElementIndex;
-						CurrentMesh->GetTriangle(TriangleIndex, Vertices[0], Vertices[1], Vertices[2], ElementIndex);
+					VolumeLightingSamples.Add(CurrentMesh->LevelGuid, TArray<FVolumeLightingSample>());
+				}
+				// Tell the rasterizer we are adding samples to this mesh's LevelId
+				Rasterizer.SetLevelGuid(CurrentMesh->LevelGuid);
+				// Rasterize all triangles in the mesh
+				for (int32 TriangleIndex = 0; TriangleIndex < CurrentMesh->NumTriangles; TriangleIndex++)
+				{
+					FStaticLightingVertex Vertices[3];
+					int32 ElementIndex;
+					CurrentMesh->GetTriangle(TriangleIndex, Vertices[0], Vertices[1], Vertices[2], ElementIndex);
 
-						if (CurrentMesh->IsElementCastingShadow(ElementIndex))
+					if (CurrentMesh->IsElementCastingShadow(ElementIndex))
+					{
+						FVector2D XYPositions[3];
+						for (int32 VertIndex = 0; VertIndex < 3; VertIndex++)
 						{
-							FVector2D XYPositions[3];
-							for (int32 VertIndex = 0; VertIndex < 3; VertIndex++)
-							{
-								// Transform world space positions from [VolumeBounds.Origin - VolumeBounds.BoxExtent, VolumeBounds.Origin + VolumeBounds.BoxExtent] into [0,1]
-								const FVector4 TransformedPosition = (Vertices[VertIndex].WorldPosition - FVector4(VolumeBounds.Origin, 0.0f) + FVector4(VolumeBounds.BoxExtent, 0.0f)) / (2.0f * FVector4(VolumeBounds.BoxExtent, 1.0f));
-								// Project positions onto the XY plane and scale to the resolution determined by DynamicObjectSettings.SurfaceLightSampleSpacing
-								XYPositions[VertIndex] = FVector2D(TransformedPosition.X * RasterSizeX, TransformedPosition.Y * RasterSizeY);
-							}
+							// Transform world space positions from [VolumeBounds.Origin - VolumeBounds.BoxExtent, VolumeBounds.Origin + VolumeBounds.BoxExtent] into [0,1]
+							const FVector4 TransformedPosition = (Vertices[VertIndex].WorldPosition - FVector4(VolumeBounds.Origin, 0.0f) + FVector4(VolumeBounds.BoxExtent, 0.0f)) / (2.0f * FVector4(VolumeBounds.BoxExtent, 1.0f));
+							// Project positions onto the XY plane and scale to the resolution determined by DynamicObjectSettings.SurfaceLightSampleSpacing
+							XYPositions[VertIndex] = FVector2D(TransformedPosition.X * RasterSizeX, TransformedPosition.Y * RasterSizeY);
+						}
 
-							const FVector4 TriangleNormal = (Vertices[2].WorldPosition - Vertices[0].WorldPosition) ^ (Vertices[1].WorldPosition - Vertices[0].WorldPosition);
-							const float TriangleArea = 0.5f * TriangleNormal.Size3();
+						const FVector4 TriangleNormal = (Vertices[2].WorldPosition - Vertices[0].WorldPosition) ^ (Vertices[1].WorldPosition - Vertices[0].WorldPosition);
+						const float TriangleArea = 0.5f * TriangleNormal.Size3();
 
-							if (TriangleArea > DELTA)
+						if (TriangleArea > DELTA)
+						{
+							if (TextureMapping)
 							{
-								if (TextureMapping)
+								// Triangle vertices in lightmap UV space, scaled by the lightmap resolution
+								const FVector2D Vertex0 = Vertices[0].TextureCoordinates[TextureMapping->LightmapTextureCoordinateIndex] * FVector2D(TextureMapping->SizeX, TextureMapping->SizeY);
+								const FVector2D Vertex1 = Vertices[1].TextureCoordinates[TextureMapping->LightmapTextureCoordinateIndex] * FVector2D(TextureMapping->SizeX, TextureMapping->SizeY);
+								const FVector2D Vertex2 = Vertices[2].TextureCoordinates[TextureMapping->LightmapTextureCoordinateIndex] * FVector2D(TextureMapping->SizeX, TextureMapping->SizeY);
+
+								// Area in lightmap space, or the number of lightmap texels covered by this triangle
+								const float LightmapTriangleArea = FMath::Abs(
+									Vertex0.X * (Vertex1.Y - Vertex2.Y)
+									+ Vertex1.X * (Vertex2.Y - Vertex0.Y)
+									+ Vertex2.X * (Vertex0.Y - Vertex1.Y));
+
+								const float TexelDensity = LightmapTriangleArea / TriangleArea;
+								// Skip texture lightmapped triangles whose texel density is less than one texel per the area of a right triangle formed by SurfaceLightSampleSpacing.
+								// If surface lighting is being calculated at a low resolution, it's unlikely that the volume near that surface needs to have detailed lighting.
+								if (TexelDensity < 2.0f / FMath::Square(DynamicObjectSettings.SurfaceLightSampleSpacing))
 								{
-									// Triangle vertices in lightmap UV space, scaled by the lightmap resolution
-									const FVector2D Vertex0 = Vertices[0].TextureCoordinates[TextureMapping->LightmapTextureCoordinateIndex] * FVector2D(TextureMapping->SizeX, TextureMapping->SizeY);
-									const FVector2D Vertex1 = Vertices[1].TextureCoordinates[TextureMapping->LightmapTextureCoordinateIndex] * FVector2D(TextureMapping->SizeX, TextureMapping->SizeY);
-									const FVector2D Vertex2 = Vertices[2].TextureCoordinates[TextureMapping->LightmapTextureCoordinateIndex] * FVector2D(TextureMapping->SizeX, TextureMapping->SizeY);
-
-									// Area in lightmap space, or the number of lightmap texels covered by this triangle
-									const float LightmapTriangleArea = FMath::Abs(
-										Vertex0.X * (Vertex1.Y - Vertex2.Y)
-										+ Vertex1.X * (Vertex2.Y - Vertex0.Y)
-										+ Vertex2.X * (Vertex0.Y - Vertex1.Y));
-
-									const float TexelDensity = LightmapTriangleArea / TriangleArea;
-									// Skip texture lightmapped triangles whose texel density is less than one texel per the area of a right triangle formed by SurfaceLightSampleSpacing.
-									// If surface lighting is being calculated at a low resolution, it's unlikely that the volume near that surface needs to have detailed lighting.
-									if (TexelDensity < 2.0f / FMath::Square(DynamicObjectSettings.SurfaceLightSampleSpacing))
-									{
-										continue;
-									}
+									continue;
 								}
+							}
 
-								// Only rasterize upward facing triangles
-								if (TriangleNormal.Z > 0.0f)
-								{
-									Rasterizer.DrawTriangle(
-										Vertices[0],
-										Vertices[1],
-										Vertices[2],
-										XYPositions[0],
-										XYPositions[1],
-										XYPositions[2],
-										false
-										);
-								}
+							// Only rasterize upward facing triangles
+							if (TriangleNormal.Z > 0.0f)
+							{
+								Rasterizer.DrawTriangle(
+									Vertices[0],
+									Vertices[1],
+									Vertices[2],
+									XYPositions[0],
+									XYPositions[1],
+									XYPositions[2],
+									false
+									);
 							}
 						}
 					}
 				}
 			}
+		}
 
-			const float DetailVolumeSpacing = DynamicObjectSettings.DetailVolumeSampleSpacing;
-			// Generate samples in a uniform 3d grid inside the detail volumes.  These will handle detail indirect lighting in areas that aren't directly above a surface.
-			for (int32 VolumeIndex = 0; VolumeIndex < Scene.CharacterIndirectDetailVolumes.Num(); VolumeIndex++)
+		const float DetailVolumeSpacing = DynamicObjectSettings.DetailVolumeSampleSpacing;
+		// Generate samples in a uniform 3d grid inside the detail volumes.  These will handle detail indirect lighting in areas that aren't directly above a surface.
+		for (int32 VolumeIndex = 0; VolumeIndex < Scene.CharacterIndirectDetailVolumes.Num(); VolumeIndex++)
+		{
+			const FBox& DetailVolumeBounds = Scene.CharacterIndirectDetailVolumes[VolumeIndex];
+			for (float SampleX = DetailVolumeBounds.Min.X; SampleX < DetailVolumeBounds.Max.X; SampleX += DetailVolumeSpacing)
 			{
-				const FBox& DetailVolumeBounds = Scene.CharacterIndirectDetailVolumes[VolumeIndex];
-				for (float SampleX = DetailVolumeBounds.Min.X; SampleX < DetailVolumeBounds.Max.X; SampleX += DetailVolumeSpacing)
+				for (float SampleY = DetailVolumeBounds.Min.Y; SampleY < DetailVolumeBounds.Max.Y; SampleY += DetailVolumeSpacing)
 				{
-					for (float SampleY = DetailVolumeBounds.Min.Y; SampleY < DetailVolumeBounds.Max.Y; SampleY += DetailVolumeSpacing)
+					for (float SampleZ = DetailVolumeBounds.Min.Z; SampleZ < DetailVolumeBounds.Max.Z; SampleZ += DetailVolumeSpacing)
 					{
-						for (float SampleZ = DetailVolumeBounds.Min.Z; SampleZ < DetailVolumeBounds.Max.Z; SampleZ += DetailVolumeSpacing)
-						{
-							const FVector4 SamplePosition(SampleX, SampleY, SampleZ);
+						const FVector4 SamplePosition(SampleX, SampleY, SampleZ);
 							
-							// Only place a sample if there are no surface lighting samples nearby
-							if (!FindNearbyVolumeSample(VolumeLightingOctree, SamplePosition, DynamicObjectSettings.SurfaceLightSampleSpacing))
+						// Only place a sample if there are no surface lighting samples nearby
+						if (!FindNearbyVolumeSample(VolumeLightingOctree, SamplePosition, DynamicObjectSettings.SurfaceLightSampleSpacing))
+						{
+							const FLightRay Ray(
+								SamplePosition,
+								SamplePosition - FVector4(0,0,VolumeBounds.BoxExtent.Z * 2.0f),
+								nullptr,
+								nullptr
+								);
+							FLightRayIntersection Intersection;
+							// Trace a ray straight down to find which level's geometry we are over, 
+							// Since this is how Dynamic Light Environments figure out which level to interpolate indirect lighting from.
+							//@todo - could probably reuse the ray trace results for all samples of the same X and Y
+							AggregateMesh.IntersectLightRay(Ray, true, false, false, MappingContext.RayCache, Intersection);
+
+							// Place the sample in the intersected level, or the persistent level if there was no intersection
+							const FGuid LevelGuid = Intersection.bIntersects ? Intersection.Mesh->LevelGuid : FGuid(0,0,0,0);
+							TArray<FVolumeLightingSample>* VolumeLightingSampleArray = VolumeLightingSamples.Find(LevelGuid);
+							if (!VolumeLightingSampleArray)
 							{
-								const FLightRay Ray(
-									SamplePosition,
-									SamplePosition - FVector4(0,0,VolumeBounds.BoxExtent.Z * 2.0f),
-									nullptr,
-									nullptr
-									);
-								FLightRayIntersection Intersection;
-								// Trace a ray straight down to find which level's geometry we are over, 
-								// Since this is how Dynamic Light Environments figure out which level to interpolate indirect lighting from.
-								//@todo - could probably reuse the ray trace results for all samples of the same X and Y
-								AggregateMesh.IntersectLightRay(Ray, true, false, false, MappingContext.RayCache, Intersection);
-
-								// Place the sample in the intersected level, or the persistent level if there was no intersection
-								const FGuid LevelGuid = Intersection.bIntersects ? Intersection.Mesh->LevelGuid : FGuid(0,0,0,0);
-								TArray<FVolumeLightingSample>* VolumeLightingSampleArray = VolumeLightingSamples.Find(LevelGuid);
-								if (!VolumeLightingSampleArray)
-								{
-									VolumeLightingSampleArray = &VolumeLightingSamples.Add(LevelGuid, TArray<FVolumeLightingSample>());
-								}
+								VolumeLightingSampleArray = &VolumeLightingSamples.Add(LevelGuid, TArray<FVolumeLightingSample>());
+							}
 
-								// Add a sample and set its radius such that its influence touches a diagonal sample on the 3d grid.
-								VolumeLightingSampleArray->Add(FVolumeLightingSample(FVector4(SamplePosition, DetailVolumeSpacing * FMath::Sqrt(3.0f))));
-								VolumeLightingOctree.AddElement(FVolumeSampleProximityElement(VolumeLightingSampleArray->Num() - 1, *VolumeLightingSampleArray));
-								if (DynamicObjectSettings.bVisualizeVolumeLightInterpolation)
-								{
-									VolumeLightingInterpolationOctree.AddElement(FVolumeSampleInterpolationElement(VolumeLightingSampleArray->Num() - 1, *VolumeLightingSampleArray));
-								}
+							// Add a sample and set its radius such that its influence touches a diagonal sample on the 3d grid.
+							VolumeLightingSampleArray->Add(FVolumeLightingSample(FVector4(SamplePosition, DetailVolumeSpacing * FMath::Sqrt(3.0f))));
+							VolumeLightingOctree.AddElement(FVolumeSampleProximityElement(VolumeLightingSampleArray->Num() - 1, *VolumeLightingSampleArray));
+							if (DynamicObjectSettings.bVisualizeVolumeLightInterpolation)
+							{
+								VolumeLightingInterpolationOctree.AddElement(FVolumeSampleInterpolationElement(VolumeLightingSampleArray->Num() - 1, *VolumeLightingSampleArray));
 							}
 						}
 					}
 				}
 			}
+		}
 
-			int32 SurfaceSamples = 0;
-			for (TMap<FGuid,TArray<FVolumeLightingSample> >::TIterator It(VolumeLightingSamples); It; ++It)
-			{
-				SurfaceSamples += It.Value().Num();
-			}
-			Stats.NumDynamicObjectSurfaceSamples = SurfaceSamples;
+		int32 SurfaceSamples = 0;
+		for (TMap<FGuid,TArray<FVolumeLightingSample> >::TIterator It(VolumeLightingSamples); It; ++It)
+		{
+			SurfaceSamples += It.Value().Num();
+		}
+		Stats.NumDynamicObjectSurfaceSamples = SurfaceSamples;
 
-			TArray<FVolumeLightingSample>* UniformVolumeSamples = VolumeLightingSamples.Find(FGuid(0,0,0,0));
-			if (!UniformVolumeSamples)
-			{
-				UniformVolumeSamples = &VolumeLightingSamples.Add(FGuid(0,0,0,0), TArray<FVolumeLightingSample>());
-			}
+		TArray<FVolumeLightingSample>* UniformVolumeSamples = VolumeLightingSamples.Find(FGuid(0,0,0,0));
+		if (!UniformVolumeSamples)
+		{
+			UniformVolumeSamples = &VolumeLightingSamples.Add(FGuid(0,0,0,0), TArray<FVolumeLightingSample>());
+		}
 
-			const float VolumeSpacingCubed = DynamicObjectSettings.VolumeLightSampleSpacing * DynamicObjectSettings.VolumeLightSampleSpacing * DynamicObjectSettings.VolumeLightSampleSpacing;
-			int32 RequestedVolumeSamples = FMath::TruncToInt(8.0f * VolumeBounds.BoxExtent.X * VolumeBounds.BoxExtent.Y * VolumeBounds.BoxExtent.Z / VolumeSpacingCubed);
-			RequestedVolumeSamples = RequestedVolumeSamples == appTruncErrorCode ? INT_MAX : RequestedVolumeSamples;
-			float EffectiveVolumeSpacing = DynamicObjectSettings.VolumeLightSampleSpacing;
+		const float VolumeSpacingCubed = DynamicObjectSettings.VolumeLightSampleSpacing * DynamicObjectSettings.VolumeLightSampleSpacing * DynamicObjectSettings.VolumeLightSampleSpacing;
+		int32 RequestedVolumeSamples = FMath::TruncToInt(8.0f * VolumeBounds.BoxExtent.X * VolumeBounds.BoxExtent.Y * VolumeBounds.BoxExtent.Z / VolumeSpacingCubed);
+		RequestedVolumeSamples = RequestedVolumeSamples == appTruncErrorCode ? INT_MAX : RequestedVolumeSamples;
+		float EffectiveVolumeSpacing = DynamicObjectSettings.VolumeLightSampleSpacing;
 
-			// Clamp the number of volume samples generated to DynamicObjectSettings.MaxVolumeSamples if necessary by resizing EffectiveVolumeSpacing
-			if (RequestedVolumeSamples > DynamicObjectSettings.MaxVolumeSamples)
-			{
-				EffectiveVolumeSpacing = FMath::Pow(8.0f * VolumeBounds.BoxExtent.X * VolumeBounds.BoxExtent.Y * VolumeBounds.BoxExtent.Z / DynamicObjectSettings.MaxVolumeSamples, .3333333f);
-			}
+		// Clamp the number of volume samples generated to DynamicObjectSettings.MaxVolumeSamples if necessary by resizing EffectiveVolumeSpacing
+		if (RequestedVolumeSamples > DynamicObjectSettings.MaxVolumeSamples)
+		{
+			EffectiveVolumeSpacing = FMath::Pow(8.0f * VolumeBounds.BoxExtent.X * VolumeBounds.BoxExtent.Y * VolumeBounds.BoxExtent.Z / DynamicObjectSettings.MaxVolumeSamples, .3333333f);
+		}
 			
-			int32 NumUniformVolumeSamples = 0;
-			// Generate samples in a uniform 3d grid inside the importance volume.  These will be used for low resolution lighting in unimportant areas.
-			for (float SampleX = VolumeBounds.Origin.X - VolumeBounds.BoxExtent.X; SampleX < VolumeBounds.Origin.X + VolumeBounds.BoxExtent.X; SampleX += EffectiveVolumeSpacing)
+		int32 NumUniformVolumeSamples = 0;
+		// Generate samples in a uniform 3d grid inside the importance volume.  These will be used for low resolution lighting in unimportant areas.
+		for (float SampleX = VolumeBounds.Origin.X - VolumeBounds.BoxExtent.X; SampleX < VolumeBounds.Origin.X + VolumeBounds.BoxExtent.X; SampleX += EffectiveVolumeSpacing)
+		{
+			for (float SampleY = VolumeBounds.Origin.Y - VolumeBounds.BoxExtent.Y; SampleY < VolumeBounds.Origin.Y + VolumeBounds.BoxExtent.Y; SampleY += EffectiveVolumeSpacing)
 			{
-				for (float SampleY = VolumeBounds.Origin.Y - VolumeBounds.BoxExtent.Y; SampleY < VolumeBounds.Origin.Y + VolumeBounds.BoxExtent.Y; SampleY += EffectiveVolumeSpacing)
+				for (float SampleZ = VolumeBounds.Origin.Z - VolumeBounds.BoxExtent.Z; SampleZ < VolumeBounds.Origin.Z + VolumeBounds.BoxExtent.Z; SampleZ += EffectiveVolumeSpacing)
 				{
-					for (float SampleZ = VolumeBounds.Origin.Z - VolumeBounds.BoxExtent.Z; SampleZ < VolumeBounds.Origin.Z + VolumeBounds.BoxExtent.Z; SampleZ += EffectiveVolumeSpacing)
+					const FVector4 SamplePosition(SampleX, SampleY, SampleZ);
+					// Only place inside the importance volume
+					if (IsPointInImportanceVolume(SamplePosition)
+						// Only place a sample if there are no surface lighting samples nearby
+						&& !FindNearbyVolumeSample(VolumeLightingOctree, SamplePosition, DynamicObjectSettings.SurfaceLightSampleSpacing))
 					{
-						const FVector4 SamplePosition(SampleX, SampleY, SampleZ);
-						// Only place inside the importance volume
-						if (IsPointInImportanceVolume(SamplePosition)
-							// Only place a sample if there are no surface lighting samples nearby
-							&& !FindNearbyVolumeSample(VolumeLightingOctree, SamplePosition, DynamicObjectSettings.SurfaceLightSampleSpacing))
+						NumUniformVolumeSamples++;
+						// Add a sample and set its radius such that its influence touches a diagonal sample on the 3d grid.
+						UniformVolumeSamples->Add(FVolumeLightingSample(FVector4(SamplePosition, EffectiveVolumeSpacing * FMath::Sqrt(3.0f))));
+						VolumeLightingOctree.AddElement(FVolumeSampleProximityElement(UniformVolumeSamples->Num() - 1, *UniformVolumeSamples));
+						if (DynamicObjectSettings.bVisualizeVolumeLightInterpolation)
 						{
-							NumUniformVolumeSamples++;
-							// Add a sample and set its radius such that its influence touches a diagonal sample on the 3d grid.
-							UniformVolumeSamples->Add(FVolumeLightingSample(FVector4(SamplePosition, EffectiveVolumeSpacing * FMath::Sqrt(3.0f))));
-							VolumeLightingOctree.AddElement(FVolumeSampleProximityElement(UniformVolumeSamples->Num() - 1, *UniformVolumeSamples));
-							if (DynamicObjectSettings.bVisualizeVolumeLightInterpolation)
-							{
-								VolumeLightingInterpolationOctree.AddElement(FVolumeSampleInterpolationElement(UniformVolumeSamples->Num() - 1, *UniformVolumeSamples));
-							}
+							VolumeLightingInterpolationOctree.AddElement(FVolumeSampleInterpolationElement(UniformVolumeSamples->Num() - 1, *UniformVolumeSamples));
 						}
 					}
 				}
 			}
-			Stats.NumDynamicObjectVolumeSamples = NumUniformVolumeSamples;
-
-			FLMRandomStream RandomStream(0);
-
-			TArray<FVector4> UniformHemisphereSamples;
-			TArray<FVector2D> UniformHemisphereSampleUniforms;
-			const int32 NumUpperVolumeSamples = ImportanceTracingSettings.NumHemisphereSamples * DynamicObjectSettings.NumHemisphereSamplesScale;
-			// Volume samples don't do any importance sampling so they need more samples for the same amount of variance as surface samples
-			const float NumThetaStepsFloat = FMath::Sqrt(NumUpperVolumeSamples / (float)PI);
-			const int32 NumThetaSteps = FMath::TruncToInt(NumThetaStepsFloat);
-			const int32 NumPhiSteps = FMath::TruncToInt(NumThetaStepsFloat * (float)PI);
+		}
 
-			GenerateStratifiedUniformHemisphereSamples(NumThetaSteps, NumPhiSteps, RandomStream, UniformHemisphereSamples, UniformHemisphereSampleUniforms);
+		Stats.NumDynamicObjectVolumeSamples = NumUniformVolumeSamples;
+		const int32 VolumeSampleTaskSize = 256;
 
-			FVector4 CombinedVector(0);
+		for (TMap<FGuid,TArray<FVolumeLightingSample> >::TIterator It(VolumeLightingSamples); It; ++It)
+		{
+			const TArray<FVolumeLightingSample>& CurrentVolumeSamples = It.Value();
 
-			for (int32 SampleIndex = 0; SampleIndex < UniformHemisphereSamples.Num(); SampleIndex++)
+			for (int32 TaskIndex = 0; TaskIndex < FMath::DivideAndRoundUp(CurrentVolumeSamples.Num(), VolumeSampleTaskSize); TaskIndex++)
 			{
-				CombinedVector += UniformHemisphereSamples[SampleIndex];
+				const int32 NumTaskSamples = FMath::Min(VolumeSampleTaskSize, CurrentVolumeSamples.Num() - TaskIndex * VolumeSampleTaskSize);
+				VolumeSampleTasks.Push(FVolumeSamplesTaskDescription(It.Key(), TaskIndex * VolumeSampleTaskSize, NumTaskSamples));
 			}
+		}
 
-			float MaxUnoccludedLength = (CombinedVector / UniformHemisphereSamples.Num()).Size3();
+		Stats.VolumeSamplePlacementThreadTime = FPlatformTime::Seconds() - VolumeSampleStartTime;
 
-			// Calculate incident radiance for each volume lighting sample
-			for (TMap<FGuid,TArray<FVolumeLightingSample> >::TIterator It(VolumeLightingSamples); It; ++It)
-			{
-				TArray<FVolumeLightingSample>& CurrentLevelSamples = It.Value();
-				for (int32 SampleIndex = 0; SampleIndex < CurrentLevelSamples.Num(); SampleIndex++)
-				{
-					FVolumeLightingSample& CurrentSample = CurrentLevelSamples[SampleIndex];
-					if (GeneralSettings.NumIndirectLightingBounces > 0 
-						// Calculating incident radiance for volume samples requires final gathering, since photons are only stored on surfaces.
-						&& (!PhotonMappingSettings.bUsePhotonMapping || PhotonMappingSettings.bUseFinalGathering))
-					{
-						const bool bDebugSamples = false;
-						CalculateVolumeSampleIncidentRadiance(UniformHemisphereSamples, MaxUnoccludedLength, CurrentSample, RandomStream, MappingContext, bDebugSamples);
-					}
-#if ALLOW_LIGHTMAP_SAMPLE_DEBUGGING
-					if (Scene.DebugMapping && DynamicObjectSettings.bVisualizeVolumeLightSamples)
-					{
-						FSHVectorRGB2 IncidentRadiance;
-						CurrentSample.ToSHVector(IncidentRadiance);
-						VolumeLightingDebugOutput.VolumeLightingSamples.Add(FDebugVolumeLightingSample(CurrentSample.PositionAndRadius, IncidentRadiance.CalcIntegral() / FSHVector2::ConstantBasisIntegral));
-					}
-#endif
-				}
-			}
+		// Make sure writes to VolumeSampleTasks are complete
+		FPlatformMisc::MemoryBarrier();
+		FPlatformAtomics::InterlockedExchange(&NumVolumeSampleTasksOutstanding, VolumeSampleTasks.Num());
+	}
+}
 
-			MappingContext.Stats.TotalVolumeSampleLightingThreadTime += FPlatformTime::Seconds() - SampleVolumeStart;
-		}
+void FStaticLightingSystem::ProcessVolumeSamplesTask(const FVolumeSamplesTaskDescription& Task)
+{
+	const double VolumeSampleStartTime = FPlatformTime::Seconds();
+
+	FLMRandomStream RandomStream(0);
+	FStaticLightingMappingContext MappingContext(nullptr, *this);
+
+	TArray<FVector4> UniformHemisphereSamples;
+	TArray<FVector2D> UniformHemisphereSampleUniforms;
+	const int32 NumUpperVolumeSamples = ImportanceTracingSettings.NumHemisphereSamples * DynamicObjectSettings.NumHemisphereSamplesScale;
+	// Volume samples don't do any importance sampling so they need more samples for the same amount of variance as surface samples
+	const float NumThetaStepsFloat = FMath::Sqrt(NumUpperVolumeSamples / (float)PI);
+	const int32 NumThetaSteps = FMath::TruncToInt(NumThetaStepsFloat);
+	const int32 NumPhiSteps = FMath::TruncToInt(NumThetaStepsFloat * (float)PI);
+
+	GenerateStratifiedUniformHemisphereSamples(NumThetaSteps, NumPhiSteps, RandomStream, UniformHemisphereSamples, UniformHemisphereSampleUniforms);
 
-		Stats.VolumeSampleThreadTime = FPlatformTime::Seconds() - VolumeSampleStartTime;
+	FVector4 CombinedVector(0);
+
+	for (int32 SampleIndex = 0; SampleIndex < UniformHemisphereSamples.Num(); SampleIndex++)
+	{
+		CombinedVector += UniformHemisphereSamples[SampleIndex];
 	}
+
+	float MaxUnoccludedLength = (CombinedVector / UniformHemisphereSamples.Num()).Size3();
+
+	TArray<FVolumeLightingSample>& CurrentLevelSamples = VolumeLightingSamples.FindChecked(Task.LevelId);
+
+	for (int32 SampleIndex = Task.StartIndex; SampleIndex < Task.StartIndex + Task.NumSamples; SampleIndex++)
+	{
+		FVolumeLightingSample& CurrentSample = CurrentLevelSamples[SampleIndex];
+
+		if (GeneralSettings.NumIndirectLightingBounces > 0 
+			// Calculating incident radiance for volume samples requires final gathering, since photons are only stored on surfaces.
+			&& (!PhotonMappingSettings.bUsePhotonMapping || PhotonMappingSettings.bUseFinalGathering))
+		{
+			const bool bDebugSamples = false;
+			CalculateVolumeSampleIncidentRadiance(UniformHemisphereSamples, MaxUnoccludedLength, CurrentSample, RandomStream, MappingContext, bDebugSamples);
+		}
+#if ALLOW_LIGHTMAP_SAMPLE_DEBUGGING
+		if (Scene.DebugMapping && DynamicObjectSettings.bVisualizeVolumeLightSamples)
+		{
+			FSHVectorRGB2 IncidentRadiance;
+			CurrentSample.ToSHVector(IncidentRadiance);
+			VolumeLightingDebugOutput.VolumeLightingSamples.Add(FDebugVolumeLightingSample(CurrentSample.PositionAndRadius, IncidentRadiance.CalcIntegral() / FSHVector2::ConstantBasisIntegral));
+		}
+#endif
+	}
+
+	MappingContext.Stats.TotalVolumeSampleLightingThreadTime += FPlatformTime::Seconds() - VolumeSampleStartTime;
 }
 
 /** 
