Fixed compile error

diff --git a/Engine/Source/Programs/UnrealLightmass/Private/Lighting/LightingSystem.h b/Engine/Source/Programs/UnrealLightmass/Private/Lighting/LightingSystem.h
--- a/Engine/Source/Programs/UnrealLightmass/Private/Lighting/LightingSystem.h
+++ b/Engine/Source/Programs/UnrealLightmass/Private/Lighting/LightingSystem.h
@@ -1899,7 +1899,7 @@ public:
 	int32 StartIndex;
 	int32 NumSamples;
 	
-	FVolumeSamplesTaskDescription(FGuid InLevelId, int32 InStartIndex, INT32 InNumSamples) :
+	FVolumeSamplesTaskDescription(FGuid InLevelId, int32 InStartIndex, int32 InNumSamples) :
 		LevelId(InLevelId),
 		StartIndex(InStartIndex),
 		NumSamples(InNumSamples)
