GenerateCoordinateSystem2 in Lightmass, slightly faster version
PR #1067

diff --git a/Engine/Source/Programs/UnrealLightmass/Private/Lighting/MonteCarlo.cpp b/Engine/Source/Programs/UnrealLightmass/Private/Lighting/MonteCarlo.cpp
--- a/Engine/Source/Programs/UnrealLightmass/Private/Lighting/MonteCarlo.cpp
+++ b/Engine/Source/Programs/UnrealLightmass/Private/Lighting/MonteCarlo.cpp
@@ -9,7 +9,6 @@ namespace Lightmass
 /** Generates valid X and Y axes of a coordinate system, given the Z axis. */
 void GenerateCoordinateSystem(const FVector4& ZAxis, FVector4& XAxis, FVector4& YAxis)
 {
-	//@todo - there is a much more efficient implementation of this in the Physically Based Rendering book
 	// Use the vector perpendicular to ZAxis and the Y axis as the XAxis
 	const FVector4 XAxisCandidate = ZAxis ^ FVector4(0,1,0);
 	if (XAxisCandidate.SizeSquared3() < KINDA_SMALL_NUMBER)
@@ -26,6 +25,28 @@ void GenerateCoordinateSystem(const FVector4& ZAxis, FVector4& XAxis, FVector4&
 	checkSlow(YAxis.IsUnit3());
 }
 
+/** Generates valid X and Y axes of a coordinate system, given the Z axis. */
+void GenerateCoordinateSystem2(const FVector4& ZAxis, FVector4& XAxis, FVector4& YAxis)
+{
+	// This implementation is based off of the one from 'Physically Based Rendering'
+	if (FMath::Abs(ZAxis.X) > FMath::Abs(ZAxis.Y))
+ 	{
+		const float InverseLength = FMath::InvSqrt(ZAxis.X * ZAxis.X + ZAxis.Z * ZAxis.Z);
+		XAxis = FVector4(-ZAxis.Z * InverseLength, 0.0f, ZAxis.X * InverseLength);
+ 	}
+ 	else
+ 	{
+		const float InverseLength = FMath::InvSqrt(ZAxis.Y * ZAxis.Y + ZAxis.Z * ZAxis.Z);
+		XAxis = FVector4(0.0f, ZAxis.Z * InverseLength, -ZAxis.Y * InverseLength);
+ 	}
+
+	YAxis = ZAxis ^ XAxis;
+ 	checkSlow(YAxis.IsUnit3());
+	checkSlow(FMath::Abs(Dot3(XAxis, ZAxis)) <= THRESH_NORMALS_ARE_ORTHOGONAL);
+	checkSlow(FMath::Abs(Dot3(YAxis, ZAxis)) <= THRESH_NORMALS_ARE_ORTHOGONAL);
+	checkSlow(FMath::Abs(Dot3(XAxis, YAxis)) <= THRESH_NORMALS_ARE_ORTHOGONAL);
+}
+
 /** Generates a pseudo-random unit vector, uniformly distributed over all directions. */
 FVector4 GetUnitVector(FLMRandomStream& RandomStream)
 {
