Lightmass triangle normal is now averaged over the lightmap texel, improves cases where a small ledge near the texel center causes the final gather to be self-occluded.
The lighting position is no longer offset when detecting intersecting surface frontfaces within a texel.  This was causing fully shadowed artifacts in cracks, as one side of the crack would move the shading position into the other.

diff --git a/Engine/Source/Editor/UnrealEd/Private/StaticLightingSystem/StaticLightingDebug.cpp b/Engine/Source/Editor/UnrealEd/Private/StaticLightingSystem/StaticLightingDebug.cpp
--- a/Engine/Source/Editor/UnrealEd/Private/StaticLightingSystem/StaticLightingDebug.cpp
+++ b/Engine/Source/Editor/UnrealEd/Private/StaticLightingSystem/StaticLightingDebug.cpp
@@ -162,7 +162,7 @@ static bool GetBarycentricWeights(
 	return false;
 }
 
-float TriangleTolerance = .01f;
+float TriangleTolerance = 1;
 
 /** Updates GCurrentSelectedLightmapSample given a selected actor's components and the location of the click. */
 void SetDebugLightmapSample(TArray<UActorComponent*>* Components, UModel* Model, int32 iSurf, FVector ClickLocation)
diff --git a/Engine/Source/Programs/UnrealLightmass/Private/Lighting/TextureMapping.cpp b/Engine/Source/Programs/UnrealLightmass/Private/Lighting/TextureMapping.cpp
--- a/Engine/Source/Programs/UnrealLightmass/Private/Lighting/TextureMapping.cpp
+++ b/Engine/Source/Programs/UnrealLightmass/Private/Lighting/TextureMapping.cpp
@@ -226,7 +226,6 @@ void FStaticLightingRasterPolicy::ProcessPixel(int32 X,int32 Y,const Interpolant
 			// Even for split texels which are mapped to triangles in different parts of the mesh.
 			TexelToVertex.MaxSampleWeight = SampleWeight;
 			TexelToVertex.WorldPosition = Interpolant.Vertex.WorldPosition;
-			TexelToVertex.TriangleNormal = TriangleNormal;
 			TexelToVertex.ElementIndex = Interpolant.ElementIndex;
 
 			for( int32 CurCoordIndex = 0; CurCoordIndex < MAX_TEXCOORDS; ++CurCoordIndex )
@@ -241,6 +240,7 @@ void FStaticLightingRasterPolicy::ProcessPixel(int32 X,int32 Y,const Interpolant
 		TexelToVertex.WorldTangentX += Interpolant.Vertex.WorldTangentX * SampleWeight;
 		TexelToVertex.WorldTangentY += Interpolant.Vertex.WorldTangentY * SampleWeight;
 		TexelToVertex.WorldTangentZ += Interpolant.Vertex.WorldTangentZ * SampleWeight;
+		TexelToVertex.TriangleNormal += TriangleNormal * SampleWeight;
 		TexelToVertex.TotalSampleWeight += SampleWeight;
 	}
 	else if (!bUseMaxWeight)
@@ -1004,6 +1004,7 @@ void FStaticLightingSystem::SetupTextureMapping(
 					TexelToVertex.WorldTangentX = TexelToVertex.WorldTangentX / TexelToVertex.TotalSampleWeight;
 					TexelToVertex.WorldTangentY = TexelToVertex.WorldTangentY / TexelToVertex.TotalSampleWeight;
 					TexelToVertex.WorldTangentZ = TexelToVertex.WorldTangentZ / TexelToVertex.TotalSampleWeight;
+					TexelToVertex.TriangleNormal = TexelToVertex.TriangleNormal / TexelToVertex.TotalSampleWeight;
 				}
 
 				// Mark the texel as mapped to some geometry in the scene
@@ -1144,12 +1145,14 @@ void FStaticLightingSystem::SetupTextureMapping(
 					// Give preference to moving the shading position outside of backfaces
 					int32 IntersectionIndexForShadingPositionMovement = ClosestBackfacingIntersectionIndex;
 
+					// Note: this is disabled as it causes problems in cracks, the lighting position will be moved inside the object
+					/*
 					// Even if we didn't hit any backfaces, still move the shading position away from an intersecting frontface if it is close enough
 					if (IntersectionIndexForShadingPositionMovement == INDEX_NONE 
 						&& ClosestIntersectionDistanceSq < (TexelToVertex.TexelRadius / 2) * (TexelToVertex.TexelRadius / 2))
 					{
 						IntersectionIndexForShadingPositionMovement = ClosestIntersectionIndex;
-					}
+					}*/
 
 					if (IntersectionIndexForShadingPositionMovement != INDEX_NONE)
 					{
