Removing Filename[0] = 0; in LMCore constructor

diff --git a/Engine/Source/Programs/UnrealLightmass/Private/LightmassCore/LMCore.cpp b/Engine/Source/Programs/UnrealLightmass/Private/LightmassCore/LMCore.cpp
--- a/Engine/Source/Programs/UnrealLightmass/Private/LightmassCore/LMCore.cpp
+++ b/Engine/Source/Programs/UnrealLightmass/Private/LightmassCore/LMCore.cpp
@@ -19,8 +19,6 @@ namespace Lightmass
 
 FLightmassLog::FLightmassLog()
 {
-	Filename[0] = 0;
-
 	// Create a Guid for this run.
 	FGuid Guid = FGuid::NewGuid();
 
