PR #5376: Clang-CL Windows build-fixes (Contributed by [REDACTED_PERSON])

#rnx

diff --git a/Engine/Source/Programs/UnrealLightmass/Private/Launch/UnrealLightmass.cpp b/Engine/Source/Programs/UnrealLightmass/Private/Launch/UnrealLightmass.cpp
--- a/Engine/Source/Programs/UnrealLightmass/Private/Launch/UnrealLightmass.cpp
+++ b/Engine/Source/Programs/UnrealLightmass/Private/Launch/UnrealLightmass.cpp
@@ -510,8 +510,10 @@ int main(int argc, ANSICHAR* argv[])
 				{
 					printf( "Exception handled in main, crash report generated, re-throwing exception\n" );
 
-					// With the crash report created, propagate the error
-					throw( 1 );
+					// With the crash report created, propagate the error. Per the MSDN documentation
+					// on GetExceptionInformation(), the flags and parameters are not available outside
+					// the exception filter (i.e. ReportCrash) as they were probably stack-pointers.
+					RaiseException(GetExceptionCode(), 0, 0, nullptr);
 				}
 			}
 			__except( EXCEPTION_EXECUTE_HANDLER )
