Fix mismatched zlib versions between UE4Editor and UnrealLightmass, which might be the reason of decompression failure in FLightmassProcessor::ImportLightMapData2DData().

diff --git a/Engine/Source/Programs/UnrealLightmass/Private/Lighting/LightmapData.cpp b/Engine/Source/Programs/UnrealLightmass/Private/Lighting/LightmapData.cpp
--- a/Engine/Source/Programs/UnrealLightmass/Private/Lighting/LightmapData.cpp
+++ b/Engine/Source/Programs/UnrealLightmass/Private/Lighting/LightmapData.cpp
@@ -3,14 +3,11 @@
 #include "LightmapData.h"
 #include "Exporter.h"
 #include "LightmassSwarm.h"
+#include "Misc/Compression.h"
 
 /** Maximum light intensity stored in vertex/ texture lightmaps. */
 #define MAX_LIGHT_INTENSITY	16.f
 
-THIRD_PARTY_INCLUDES_START
-#include "ThirdParty/zlib/zlib-1.2.5/Inc/zlib.h"
-THIRD_PARTY_INCLUDES_END
-
 namespace Lightmass
 {
 	/**
@@ -30,18 +27,17 @@ namespace Lightmass
 			return;
 		}
 
-		/** Get's zlib's max size needed (as seen at http://www.zlib.net/zlib_tech.html) */
-		#define CALC_ZLIB_MAX(x) (x + (((x + 16383) / 16384) * 5 + 6))
+		check(UncompressedDataSize <= INT_MAX);
 
 		// allocate all of the input space for the output, with extra space for max overhead of zlib (when compressed > uncompressed)
-		unsigned long CompressedSize = CALC_ZLIB_MAX(UncompressedDataSize);
+		int32 CompressedSize = FCompression::CompressMemoryBound(NAME_Zlib, (int32)UncompressedDataSize);
 		CompressedBuffer = (uint8*)FMemory::Malloc(CompressedSize);
 
-		// compress the data
-		int32 Err = compress(CompressedBuffer, &CompressedSize, UncompressedBuffer, UncompressedDataSize);
+		// compress the data		
+		bool bSucceeded = FCompression::CompressMemory(NAME_Zlib, CompressedBuffer, CompressedSize, UncompressedBuffer, UncompressedDataSize);
 
 		// if it failed send the data uncompressed, which we mark by setting compressed size to 0
-		checkf(Err == Z_OK, TEXT("zlib failed to compress, which is very unexpected (err = %d)"), Err);
+		checkf(bSucceeded, TEXT("zlib failed to compress, which is very unexpected"));
 
 		// cache the compressed size in the header so the other side knows how much to read
 		CompressedDataSize = CompressedSize;
