Cherry-pick 7106438: Fix Lightmass crash introduced in CL 6805385

diff --git a/Engine/Source/Programs/UnrealLightmass/Private/Lighting/AdaptiveVolumetricLightmap.cpp b/Engine/Source/Programs/UnrealLightmass/Private/Lighting/AdaptiveVolumetricLightmap.cpp
--- a/Engine/Source/Programs/UnrealLightmass/Private/Lighting/AdaptiveVolumetricLightmap.cpp
+++ b/Engine/Source/Programs/UnrealLightmass/Private/Lighting/AdaptiveVolumetricLightmap.cpp
@@ -636,14 +636,20 @@ void FStaticLightingSystem::ProcessVolumetricLightmapBrickTask(FVolumetricLightm
 	}
 }
 
-void FStaticLightingSystem::ProcessVolumetricLightmapTaskIfAvailable()
+bool FStaticLightingSystem::ProcessVolumetricLightmapTaskIfAvailable()
 {
+	bool bAnyTaskProcessedByThisThread = false;
+
 	while (FVolumetricLightmapBrickTaskDescription * NextTask = VolumetricLightmapBrickTasks.Pop())
 	{
 		//UE_LOG(LogLightmass, Warning, TEXT("Thread picked up volumetric lightmap task"));
 		ProcessVolumetricLightmapBrickTask(NextTask);
 		FPlatformAtomics::InterlockedDecrement(NextTask->NumOutstandingBrickTasks);
+
+		bAnyTaskProcessedByThisThread = true;
 	}
+
+	return bAnyTaskProcessedByThisThread;
 }
 
 void FStaticLightingSystem::GenerateVoxelTestPositions(TArray<FVector>& VoxelTestPositions) const
diff --git a/Engine/Source/Programs/UnrealLightmass/Private/Lighting/LightingSystem.cpp b/Engine/Source/Programs/UnrealLightmass/Private/Lighting/LightingSystem.cpp
--- a/Engine/Source/Programs/UnrealLightmass/Private/Lighting/LightingSystem.cpp
+++ b/Engine/Source/Programs/UnrealLightmass/Private/Lighting/LightingSystem.cpp
@@ -1586,34 +1586,63 @@ void FStaticLightingSystem::ThreadLoop(bool bIsMainThread, int32 ThreadIndex, FT
 	{
 		const double StartLoopTime = FPlatformTime::Seconds();
 
-		while (FCacheIndirectTaskDescription * NextCacheTask = CacheIndirectLightingTasks.Pop())
-		{
-			//UE_LOG(LogLightmass, Warning, TEXT("Thread %u picked up Cache Indirect task for %u"), ThreadIndex, NextCacheTask->TextureMapping->Guid.D);
-			ProcessCacheIndirectLightingTask(NextCacheTask, false);
-			NextCacheTask->TextureMapping->CompletedCacheIndirectLightingTasks.Push(NextCacheTask);
-			FPlatformAtomics::InterlockedDecrement(&NextCacheTask->TextureMapping->NumOutstandingCacheTasks);
-		}
+		bool bAnyTaskProcessedByThisThread = false;
 
-		while (FInterpolateIndirectTaskDescription * NextInterpolateTask = InterpolateIndirectLightingTasks.Pop())
+		// Process any existing local tasks before fetching new tasks from Swarm
 		{
-			//UE_LOG(LogLightmass, Warning, TEXT("Thread %u picked up Interpolate indirect task for %u"), ThreadIndex, NextInterpolateTask->TextureMapping->Guid.D);
-			ProcessInterpolateTask(NextInterpolateTask, false);
-			NextInterpolateTask->TextureMapping->CompletedInterpolationTasks.Push(NextInterpolateTask);
-			FPlatformAtomics::InterlockedDecrement(&NextInterpolateTask->TextureMapping->NumOutstandingInterpolationTasks);
-		}
+			while (FCacheIndirectTaskDescription * NextCacheTask = CacheIndirectLightingTasks.Pop())
+			{
+				//UE_LOG(LogLightmass, Warning, TEXT("Thread %u picked up Cache Indirect task for %u"), ThreadIndex, NextCacheTask->TextureMapping->Guid.D);
+				ProcessCacheIndirectLightingTask(NextCacheTask, false);
+				NextCacheTask->TextureMapping->CompletedCacheIndirectLightingTasks.Push(NextCacheTask);
+				FPlatformAtomics::InterlockedDecrement(&NextCacheTask->TextureMapping->NumOutstandingCacheTasks);
 
-		ProcessVolumetricLightmapTaskIfAvailable();
+				bAnyTaskProcessedByThisThread = true;
+			}
 
-		if (NumOutstandingVolumeDataLayers > 0)
-		{
-			const int32 ThreadZ = FPlatformAtomics::InterlockedIncrement(&OutstandingVolumeDataLayerIndex);
-			if (ThreadZ < VolumeSizeZ)
+			while (FInterpolateIndirectTaskDescription * NextInterpolateTask = InterpolateIndirectLightingTasks.Pop())
 			{
-				CalculateVolumeDistanceFieldWorkRange(ThreadZ);
-				const int32 NumTasksRemaining = FPlatformAtomics::InterlockedDecrement(&NumOutstandingVolumeDataLayers);
-				if (NumTasksRemaining == 0)
+				//UE_LOG(LogLightmass, Warning, TEXT("Thread %u picked up Interpolate indirect task for %u"), ThreadIndex, NextInterpolateTask->TextureMapping->Guid.D);
+				ProcessInterpolateTask(NextInterpolateTask, false);
+				NextInterpolateTask->TextureMapping->CompletedInterpolationTasks.Push(NextInterpolateTask);
+				FPlatformAtomics::InterlockedDecrement(&NextInterpolateTask->TextureMapping->NumOutstandingInterpolationTasks);
+
+				bAnyTaskProcessedByThisThread = true;
+			}
+
+			bAnyTaskProcessedByThisThread |= ProcessVolumetricLightmapTaskIfAvailable();
+
+			while (NumOutstandingVolumeDataLayers > 0)
+			{
+				const int32 ThreadZ = FPlatformAtomics::InterlockedIncrement(&OutstandingVolumeDataLayerIndex);
+				if (ThreadZ < VolumeSizeZ)
+				{
+					CalculateVolumeDistanceFieldWorkRange(ThreadZ);
+					const int32 NumTasksRemaining = FPlatformAtomics::InterlockedDecrement(&NumOutstandingVolumeDataLayers);
+					if (NumTasksRemaining == 0)
+					{
+						FPlatformAtomics::InterlockedExchange(&bShouldExportVolumeDistanceField, true);
+					}
+
+					bAnyTaskProcessedByThisThread = true;
+				}
+			}
+
+			while (NumVolumeSampleTasksOutstanding > 0)
+			{
+				const int32 TaskIndex = FPlatformAtomics::InterlockedIncrement(&NextVolumeSampleTaskIndex);
+
+				if (TaskIndex < VolumeSampleTasks.Num())
 				{
-					FPlatformAtomics::InterlockedExchange(&bShouldExportVolumeDistanceField, true);
+					ProcessVolumeSamplesTask(VolumeSampleTasks[TaskIndex]);
+					const int32 NumTasksRemaining = FPlatformAtomics::InterlockedDecrement(&NumVolumeSampleTasksOutstanding);
+
+					if (NumTasksRemaining == 0)
+					{
+						FPlatformAtomics::InterlockedExchange(&bShouldExportVolumeSampleData, true);
+					}
+
+					bAnyTaskProcessedByThisThread = true;
 				}
 			}
 		}
@@ -1690,31 +1719,15 @@ void FStaticLightingSystem::ThreadLoop(bool bIsMainThread, int32 ThreadIndex, FT
 		}
 		else
 		{
-			if (!bSignaledMappingsComplete && NumOutstandingVolumeDataLayers <= 0)
-			{
-				bSignaledMappingsComplete = true;
-				GSwarm->SendMessage( NSwarm::FTimingMessage( NSwarm::PROGSTATE_Processing0, ThreadIndex ) );
-			}
-			
-			if (NumVolumeSampleTasksOutstanding > 0)
+			if (!bAnyTaskProcessedByThisThread && TasksInProgressThatWillNeedHelp <= 0 && NumOutstandingVolumeDataLayers <= 0 && NumVolumeSampleTasksOutstanding <= 0)
 			{
-				const int32 TaskIndex = FPlatformAtomics::InterlockedIncrement(&NextVolumeSampleTaskIndex);
-
-				if (TaskIndex < VolumeSampleTasks.Num())
+				if (!bSignaledMappingsComplete)
 				{
-					ProcessVolumeSamplesTask(VolumeSampleTasks[TaskIndex]);
-					const int32 NumTasksRemaining = FPlatformAtomics::InterlockedDecrement(&NumVolumeSampleTasksOutstanding);
-
-					if (NumTasksRemaining == 0)
-					{
-						FPlatformAtomics::InterlockedExchange(&bShouldExportVolumeSampleData, true);
-					}
+					bSignaledMappingsComplete = true;
+					GSwarm->SendMessage(NSwarm::FTimingMessage(NSwarm::PROGSTATE_Processing0, ThreadIndex));
 				}
-			}
 
-			if (NumVolumeSampleTasksOutstanding <= 0 && NumOutstandingVolumeDataLayers <= 0)
-			{
-				if (TasksInProgressThatWillNeedHelp <= 0 && !bRequestForTaskTimedOut)
+				if (!bRequestForTaskTimedOut)
 				{
 					// All mappings have been processed, so end this thread.
 					bIsDone = true;
diff --git a/Engine/Source/Programs/UnrealLightmass/Private/Lighting/LightingSystem.h b/Engine/Source/Programs/UnrealLightmass/Private/Lighting/LightingSystem.h
--- a/Engine/Source/Programs/UnrealLightmass/Private/Lighting/LightingSystem.h
+++ b/Engine/Source/Programs/UnrealLightmass/Private/Lighting/LightingSystem.h
@@ -2226,7 +2226,7 @@ private:
 
 	void ProcessVolumetricLightmapBrickTask(class FVolumetricLightmapBrickTaskDescription* Task);
 
-	void ProcessVolumetricLightmapTaskIfAvailable();
+	bool ProcessVolumetricLightmapTaskIfAvailable();
 
 	void GenerateVoxelTestPositions(TArray<FVector>& VoxelTestPositions) const;
 
diff --git a/Engine/Source/Programs/UnrealLightmass/Private/Lighting/TextureMapping.cpp b/Engine/Source/Programs/UnrealLightmass/Private/Lighting/TextureMapping.cpp
--- a/Engine/Source/Programs/UnrealLightmass/Private/Lighting/TextureMapping.cpp
+++ b/Engine/Source/Programs/UnrealLightmass/Private/Lighting/TextureMapping.cpp
@@ -552,7 +552,7 @@ void FStaticLightingSystem::AdjustRepresentativeSurfelForTexelsTextureMapping(
 						// Mark the texel as intersecting another surface so we can avoid filtering across it later
 						TexelToVertex.bIntersectingSurface = true;
 
-						TexelToVertex.TexelRadius = FMath::Min(TexelToVertex.TexelRadius, FMath::Sqrt(ClosestIntersectionDistanceSq));
+						TexelToVertex.TexelRadius = FMath::Min(TexelToVertex.TexelRadius, FMath::Sqrt(ClosestIntersectionDistanceSq / 2.0f));
 					}
 
 					// Give preference to moving the shading position outside of backfaces
@@ -579,6 +579,8 @@ void FStaticLightingSystem::AdjustRepresentativeSurfelForTexelsTextureMapping(
 
 						TexelToVertex.TexelRadius = (OffsetShadingPosition - Intersections[IntersectionIndexForShadingPositionMovement].IntersectionVertex.WorldPosition).Size3();
 					}
+
+					TexelToVertex.TexelRadius = FMath::Max(TexelToVertex.TexelRadius, SceneConstants.SmallestTexelRadius);
 				}
 			}
 			else
