Cherry pick 7223770: Fix TexelRadius being overriden before backface movement in AdjustRepresentativeSurfelForTexelsTextureMapping

diff --git a/Engine/Source/Programs/UnrealLightmass/Private/Lighting/TextureMapping.cpp b/Engine/Source/Programs/UnrealLightmass/Private/Lighting/TextureMapping.cpp
--- a/Engine/Source/Programs/UnrealLightmass/Private/Lighting/TextureMapping.cpp
+++ b/Engine/Source/Programs/UnrealLightmass/Private/Lighting/TextureMapping.cpp
@@ -531,16 +531,22 @@ void FStaticLightingSystem::AdjustRepresentativeSurfelForTexelsTextureMapping(
 						{
 							const float DistanceSquared = (Intersections[CornerIndex].IntersectionVertex.WorldPosition - TexelCenterOffset).SizeSquared3();
 
-							if (ClosestIntersectionIndex == INDEX_NONE || DistanceSquared < ClosestIntersectionDistanceSq)
+							if (!bHitBackfaces[CornerIndex] && (ClosestIntersectionIndex == INDEX_NONE || (DistanceSquared < ClosestIntersectionDistanceSq)))
 							{
 								ClosestIntersectionDistanceSq = DistanceSquared;
 								ClosestIntersectionIndex = CornerIndex;
+
+								// Mark the texel as intersecting another surface so we can avoid filtering across it later
+								TexelToVertex.bIntersectingSurface = true;
 							}
 
-							if (bHitBackfaces[CornerIndex] && DistanceSquared < ClosestBackfacingIntersectionDistanceSq)
+							if (bHitBackfaces[CornerIndex] && (DistanceSquared < ClosestBackfacingIntersectionDistanceSq && !FMath::IsNearlyEqual(DistanceSquared, ClosestBackfacingIntersectionDistanceSq, SceneConstants.SmallestTexelRadius)))
 							{
 								ClosestBackfacingIntersectionDistanceSq = DistanceSquared;
 								ClosestBackfacingIntersectionIndex = CornerIndex;
+
+								// Mark the texel as intersecting another surface so we can avoid filtering across it later
+								TexelToVertex.bIntersectingSurface = true;
 							}
 						}
 					}
@@ -549,9 +555,6 @@ void FStaticLightingSystem::AdjustRepresentativeSurfelForTexelsTextureMapping(
 					{
 						checkSlow(Intersections[ClosestIntersectionIndex].bIntersects);
 
-						// Mark the texel as intersecting another surface so we can avoid filtering across it later
-						TexelToVertex.bIntersectingSurface = true;
-
 						TexelToVertex.TexelRadius = FMath::Min(TexelToVertex.TexelRadius, FMath::Sqrt(ClosestIntersectionDistanceSq / 2.0f));
 					}
 
@@ -577,7 +580,7 @@ void FStaticLightingSystem::AdjustRepresentativeSurfelForTexelsTextureMapping(
 						// Project back onto plane of texel to avoid incorrect self occlusion
 						TexelToVertex.WorldPosition = OffsetShadingPosition + TexelToVertex.TriangleNormal * Dot3(TexelToVertex.TriangleNormal, TexelToVertex.WorldPosition - OffsetShadingPosition);
 
-						TexelToVertex.TexelRadius = (OffsetShadingPosition - Intersections[IntersectionIndexForShadingPositionMovement].IntersectionVertex.WorldPosition).Size3();
+						TexelToVertex.TexelRadius = (OffsetShadingPosition - Intersections[IntersectionIndexForShadingPositionMovement].IntersectionVertex.WorldPosition).Size3() / FMath::Sqrt(2.0f);
 					}
 
 					TexelToVertex.TexelRadius = FMath::Max(TexelToVertex.TexelRadius, SceneConstants.SmallestTexelRadius);
