Fixed rect light lightmass issues.

diff --git a/Engine/Shaders/Private/VolumetricFog.usf b/Engine/Shaders/Private/VolumetricFog.usf
--- a/Engine/Shaders/Private/VolumetricFog.usf
+++ b/Engine/Shaders/Private/VolumetricFog.usf
@@ -263,8 +263,6 @@ void InjectShadowedLocalLightPS(
 			if( LightData.bRectLight )
 			{
 				FRect Rect = GetRect( ToLight, LightData );
-				// TODO Move CPU side
-				LightColor /= 2 * Rect.Extent.x * Rect.Extent.y;
 
 				Lighting = IntegrateLight( Rect );
 			}
@@ -595,8 +593,6 @@ void LightScatteringCS(
 				if( LightData.bRectLight )
 				{
 					FRect Rect = GetRect( ToLight, LightData );
-					// TODO Move CPU side
-					LightColor /= 2 * Rect.Extent.x * Rect.Extent.y;
 
 					Lighting = IntegrateLight( Rect );
 				}
diff --git a/Engine/Source/Editor/UnrealEd/Private/Lightmass/Lightmass.cpp b/Engine/Source/Editor/UnrealEd/Private/Lightmass/Lightmass.cpp
--- a/Engine/Source/Editor/UnrealEd/Private/Lightmass/Lightmass.cpp
+++ b/Engine/Source/Editor/UnrealEd/Private/Lightmass/Lightmass.cpp
@@ -1103,7 +1103,6 @@ void FLightmassExporter::WriteLights( int32 Channel )
 	{
 		const URectLightComponent* Light = RectLights[LightIndex];
 		Lightmass::FLightData LightData;
-		Lightmass::FPointLightData PointData;
 		Copy( Light, LightData );
 		LightData.IndirectLightingSaturation = Light->LightmassSettings.IndirectLightingSaturation;
 		LightData.ShadowExponent = Light->LightmassSettings.ShadowExponent;
@@ -1114,12 +1113,92 @@ void FLightmassExporter::WriteLights( int32 Channel )
 		TArray< uint8 > LightProfileTextureData;
 		CopyLightProfile( Light, LightData, LightProfileTextureData );
 
+		Lightmass::FPointLightData PointData;
 		PointData.Radius = Light->AttenuationRadius;
 		PointData.FalloffExponent = 0.0f;
 		PointData.LightTangent = Light->GetComponentTransform().GetUnitAxis(EAxis::Z);
+
+		Lightmass::FRectLightData RectData = { 0, 0, FLinearColor::White };
+		// TODO export texture data. Below code is written for that in mind but source data doesn't contain mips.
+		//TArray< FFloat16Color > SourceTexture;
+		if( Light->SourceTexture )
+		{
+			FTextureSource& Source = Light->SourceTexture->Source;
+
+			if( Source.IsValid() )
+			{
+				ETextureSourceFormat Format = Source.GetFormat();
+				int32 BytesPerPixel = Source.GetBytesPerPixel();
+				bool SRGB = Light->SourceTexture->SRGB;
+
+				uint32 NumMips = Source.GetNumMips();
+				uint32 MipLevel = NumMips - FMath::Min( NumMips, 8u );
+
+				uint32 SizeX = FMath::Max( Source.GetSizeX() >> MipLevel, 1 );
+				uint32 SizeY = FMath::Max( Source.GetSizeY() >> MipLevel, 1 );
+
+				RectData.SourceTextureSizeX = SizeX;
+				RectData.SourceTextureSizeY = SizeY;
+
+				while( SizeX > 1 || SizeY > 1 )
+				{
+					//int32 i = SourceTexture.AddUninitialized( SizeX * SizeY );
+					
+					TArray< uint8 > MipData;
+					Source.GetMipData( MipData, MipLevel );
+
+					uint8* Pixel = MipData.GetData();
+					uint8* PixelEnd = Pixel + MipData.Num();
+					while( Pixel < PixelEnd )
+					{
+						if( Format == TSF_RGBA16 || Format == TSF_RGBA16F )
+						{
+							//SourceTexture[ i++ ] = *(FFloat16Color*)Pixel;
+							RectData.SourceTextureAvgColor += FLinearColor( *(FFloat16Color*)Pixel );
+						}
+						else
+						{
+							FColor Color = FColor( Pixel[0], Pixel[0], Pixel[0] );
+							if( Format == TSF_BGRA8 || Format == TSF_BGRE8 )
+							{
+								Color = *(FColor*)Pixel;
+							}
+
+							FLinearColor LinearColor;
+							if( SRGB )
+							{
+								LinearColor = FLinearColor( Color );
+							}
+							else
+							{
+								LinearColor.R = float( Color.R ) / 255.0f;
+								LinearColor.G = float( Color.G ) / 255.0f;
+								LinearColor.B = float( Color.B ) / 255.0f;
+							}
+							
+							//SourceTexture[ i++ ] = FFloat16Color( LinearColor );
+							RectData.SourceTextureAvgColor += LinearColor;
+						}
+
+						Pixel += BytesPerPixel;
+					}
+
+					// Source doesn't often have mips just average color for now.
+					RectData.SourceTextureAvgColor *= 1.0f / ( SizeX * SizeY );
+					break;
+					
+					SizeX = SizeX > 1 ? SizeX >> 1 : 1;
+					SizeY = SizeY > 1 ? SizeY >> 1 : 1;
+					MipLevel++;
+				}
+			}
+		}
+
 		Swarm.WriteChannel( Channel, &LightData, sizeof(LightData) );
 		Swarm.WriteChannel( Channel, LightProfileTextureData.GetData(), LightProfileTextureData.Num() * LightProfileTextureData.GetTypeSize() );
 		Swarm.WriteChannel( Channel, &PointData, sizeof(PointData) );
+		Swarm.WriteChannel( Channel, &RectData, sizeof(RectData) );
+		//Swarm.WriteChannel( Channel, SourceTexture.GetData(), SourceTexture.Num() * SourceTexture.GetTypeSize() );
 		UpdateExportProgress();
 	}
 
diff --git a/Engine/Source/Programs/UnrealLightmass/Private/ImportExport/LightmassScene.cpp b/Engine/Source/Programs/UnrealLightmass/Private/ImportExport/LightmassScene.cpp
--- a/Engine/Source/Programs/UnrealLightmass/Private/ImportExport/LightmassScene.cpp
+++ b/Engine/Source/Programs/UnrealLightmass/Private/ImportExport/LightmassScene.cpp
@@ -1549,6 +1549,30 @@ void FSpotLight::SampleDirection(
 //----------------------------------------------------------------------------
 //	Rect light class
 //----------------------------------------------------------------------------
+void FRectLight::Import( FLightmassImporter& Importer )
+{
+	FPointLight::Import( Importer );
+
+	Importer.ImportData( (FRectLightData*)this );
+
+#if 0
+	uint32 NumPixels = 0;
+	if( SourceTextureSizeX + SourceTextureSizeY > 0 )
+	{
+		uint32 SizeX = SourceTextureSizeX;
+		uint32 SizeY = SourceTextureSizeY;
+		while( SizeX > 1 || SizeY > 1 )
+		{
+			NumPixels += SizeX * SizeY;
+
+			SizeX = SizeX > 1 ? SizeX >> 1 : 1;
+			SizeY = SizeY > 1 ? SizeY >> 1 : 1;
+		}
+	}
+	Importer.ImportArray( SourceTexture, NumPixels );
+#endif
+}
+
 
 FVector PolygonIrradiance( FVector Poly[4] )
 {
@@ -1605,7 +1629,18 @@ FLinearColor FRectLight::GetDirectIntensity(const FVector4& Point, bool bCalcula
 	float LightRadiusMask = FMath::Square(FMath::Max(0.0f, 1.0f - FMath::Square(DistanceSqr / (Radius * Radius))));
 	DistanceAttenuation *= LightRadiusMask;
 
-	return FLight::GetDirectIntensity(Point, bCalculateForIndirectLighting) * DistanceAttenuation;
+	return FLight::GetDirectIntensity(Point, bCalculateForIndirectLighting) * DistanceAttenuation * SourceTextureAvgColor;
+}
+
+/** Returns the number of direct photons to gather required by this light. */
+int32 FRectLight::GetNumDirectPhotons(float DirectPhotonDensity) const
+{
+	const float InfluenceSphereSurfaceAreaMillions = 4.0f * (float)PI * FMath::Square(Radius) / 1000000.0f;
+	// Find the fraction of the sphere's surface area that is inside the cone
+	const float SurfaceAreaSphereFraction = 0.25f;
+	// Gather enough photons to meet DirectPhotonDensity on the spherical cap at the influence radius of the spot light.
+	const int32 NumDirectPhotons = FMath::TruncToInt(InfluenceSphereSurfaceAreaMillions * SurfaceAreaSphereFraction * DirectPhotonDensity);
+	return NumDirectPhotons == appTruncErrorCode ? INT_MAX : NumDirectPhotons;
 }
 
 /** Validates a surface sample given the position that sample is affecting. */
@@ -1707,6 +1742,55 @@ void FRectLight::SampleLightSurface(FLMRandomStream& RandomStream, FLightSurface
 	Sample.PDF = 1.0f / (4.0f * Extent.X * Extent.Y);
 }
 
+/** Generates a direction sample from the light's domain */
+void FRectLight::SampleDirection(FLMRandomStream& RandomStream, FLightRay& SampleRay, FVector4& LightSourceNormal, FVector2D& LightSurfacePosition, float& RayPDF, FLinearColor& Power) const
+{
+	FVector AxisY = GetLightTangent();
+	FVector AxisZ = -Direction;
+	FVector AxisX = AxisY ^ AxisZ;
+	FVector2D Extent(
+		LightSourceRadius,
+		LightSourceLength
+	);
+
+	float UnitX = RandomStream.GetFraction() * 2.0f - 1.0f;
+	float UnitY = RandomStream.GetFraction() * 2.0f - 1.0f;
+
+	FVector4 SamplePosition = Position + AxisX * Extent.X * UnitX + AxisY * Extent.Y * UnitY;
+
+	FVector4 SampleDirection = GetCosineHemisphereVector( RandomStream );
+	RayPDF = SampleDirection.Z / PI;
+
+	SampleDirection =
+		SampleDirection.X * AxisX +
+		SampleDirection.Y * AxisY +
+		SampleDirection.Z * AxisZ;
+
+	SampleRay = FLightRay(
+		SamplePosition,
+		SamplePosition - SampleDirection * Radius,
+		NULL,
+		this
+		);
+
+	LightSourceNormal = -AxisZ;
+
+	Power = SourceTextureAvgColor * IndirectColor * Brightness;
+}
+
+/** Generates a direction sample from the light based on the given rays */
+void FRectLight::SampleDirection(
+	const TArray<FIndirectPathRay>& IndirectPathRays, 
+	FLMRandomStream& RandomStream, 
+	FLightRay& SampleRay, 
+	float& RayPDF, 
+	FLinearColor& Power) const
+{
+	FVector4 Unused;
+	FVector2D Unused2;
+	FRectLight::SampleDirection(RandomStream, SampleRay, Unused, Unused2, RayPDF, Power);
+}
+
 
 //----------------------------------------------------------------------------
 //	Sky light class
diff --git a/Engine/Source/Programs/UnrealLightmass/Private/ImportExport/LightmassScene.h b/Engine/Source/Programs/UnrealLightmass/Private/ImportExport/LightmassScene.h
--- a/Engine/Source/Programs/UnrealLightmass/Private/ImportExport/LightmassScene.h
+++ b/Engine/Source/Programs/UnrealLightmass/Private/ImportExport/LightmassScene.h
@@ -463,17 +463,25 @@ protected:
 //----------------------------------------------------------------------------
 //	Rect light class
 //----------------------------------------------------------------------------
-class FRectLight : public FPointLight
+class FRectLight : public FPointLight, public FRectLightData
 {
 public:
+	virtual void			Import( class FLightmassImporter& Importer );
+
 	virtual FLinearColor	GetDirectIntensity(const FVector4& Point, bool bCalculateForIndirectLighting) const override;
+	int32					GetNumDirectPhotons(float DirectPhotonDensity) const override;
 	virtual void			ValidateSurfaceSample(const FVector4& Point, FLightSurfaceSample& Sample) const override;
 	virtual FVector4		LightCenterPosition(const FVector4& ReceivingPosition, const FVector4& ReceivingNormal) const override;
 	virtual bool			BehindSurface(const FVector4& TrianglePoint, const FVector4& TriangleNormal) const override;
 	virtual FVector4		GetDirectLightingDirection(const FVector4& Point, const FVector4& PointNormal) const override;
 
+	virtual void			SampleDirection(FLMRandomStream& RandomStream, FLightRay& SampleRay, FVector4& LightSourceNormal, FVector2D& LightSurfacePosition, float& RayPDF, FLinearColor& Power) const override;
+	virtual void			SampleDirection(const TArray<FIndirectPathRay>& IndirectPathRays, FLMRandomStream& RandomStream, FLightRay& SampleRay, float& RayPDF, FLinearColor& Power) const override;
+
 protected:
 	virtual void			SampleLightSurface(FLMRandomStream& RandomStream, FLightSurfaceSample& Sample) const override;
+
+	TArray< FFloat16Color >	SourceTexture;
 };
 
 
diff --git a/Engine/Source/Programs/UnrealLightmass/Public/SceneExport.h b/Engine/Source/Programs/UnrealLightmass/Public/SceneExport.h
--- a/Engine/Source/Programs/UnrealLightmass/Public/SceneExport.h
+++ b/Engine/Source/Programs/UnrealLightmass/Public/SceneExport.h
@@ -1021,6 +1021,13 @@ struct FSkyLightData
 	FSHVectorRGB3 IrradianceEnvironmentMap;
 };
 
+struct FRectLightData
+{
+	uint32			SourceTextureSizeX;
+	uint32			SourceTextureSizeY;
+	FLinearColor	SourceTextureAvgColor;
+};
+
 //----------------------------------------------------------------------------
 //	Material
 //----------------------------------------------------------------------------
