Fixed static rect lights emitting backwards light from their backsides.

diff --git a/Engine/Source/Programs/UnrealLightmass/Private/ImportExport/LightmassScene.cpp b/Engine/Source/Programs/UnrealLightmass/Private/ImportExport/LightmassScene.cpp
--- a/Engine/Source/Programs/UnrealLightmass/Private/ImportExport/LightmassScene.cpp
+++ b/Engine/Source/Programs/UnrealLightmass/Private/ImportExport/LightmassScene.cpp
@@ -1604,7 +1604,7 @@ FVector PolygonIrradiance( FVector Poly[4] )
  */
 FLinearColor FRectLight::GetDirectIntensity(const FVector4& Point, bool bCalculateForIndirectLighting) const
 {
-	FVector4 ToLight = Position - Point;
+	FVector ToLight = Position - Point;
 
 	FVector AxisY = GetLightTangent();
 	FVector AxisZ = -Direction;
@@ -1614,6 +1614,9 @@ FLinearColor FRectLight::GetDirectIntensity(const FVector4& Point, bool bCalcula
 		LightSourceLength
 	);
 
+	if( ( ToLight | AxisZ ) < 0.0f )
+		return FLinearColor::Black;
+
 	FVector Poly[4];
 	Poly[0] = ToLight - AxisX * Extent.X - AxisY * Extent.Y;
 	Poly[1] = ToLight + AxisX * Extent.X - AxisY * Extent.Y;
@@ -1625,7 +1628,7 @@ FLinearColor FRectLight::GetDirectIntensity(const FVector4& Point, bool bCalcula
 	// TODO Move CPU cide
 	DistanceAttenuation /= 2 * Extent.X * Extent.Y;
 
-	float DistanceSqr = ToLight.SizeSquared3();
+	float DistanceSqr = ToLight.SizeSquared();
 	float LightRadiusMask = FMath::Square(FMath::Max(0.0f, 1.0f - FMath::Square(DistanceSqr / (Radius * Radius))));
 	DistanceAttenuation *= LightRadiusMask;
 
@@ -1692,6 +1695,23 @@ FVector4 FRectLight::LightCenterPosition(const FVector4& ReceivingPosition, cons
 #endif
 }
 
+bool FRectLight::AffectsBounds(const FBoxSphereBounds& Bounds) const
+{
+	FPlane Plane( Position, Direction );
+	float DistanceToPlane = ( Bounds.Origin - Position ) | Direction;
+	if( DistanceToPlane < -Bounds.SphereRadius )
+	{
+		return false;
+	}
+
+	if(!FPointLight::AffectsBounds(Bounds))
+	{
+		return false;
+	}
+
+	return true;
+}
+
 /** Returns true if all parts of the light are behind the surface being tested. */
 bool FRectLight::BehindSurface(const FVector4& TrianglePoint, const FVector4& TriangleNormal) const
 {
diff --git a/Engine/Source/Programs/UnrealLightmass/Private/ImportExport/LightmassScene.h b/Engine/Source/Programs/UnrealLightmass/Private/ImportExport/LightmassScene.h
--- a/Engine/Source/Programs/UnrealLightmass/Private/ImportExport/LightmassScene.h
+++ b/Engine/Source/Programs/UnrealLightmass/Private/ImportExport/LightmassScene.h
@@ -469,9 +469,10 @@ public:
 	virtual void			Import( class FLightmassImporter& Importer );
 
 	virtual FLinearColor	GetDirectIntensity(const FVector4& Point, bool bCalculateForIndirectLighting) const override;
-	int32					GetNumDirectPhotons(float DirectPhotonDensity) const override;
+	virtual int32			GetNumDirectPhotons(float DirectPhotonDensity) const override;
 	virtual void			ValidateSurfaceSample(const FVector4& Point, FLightSurfaceSample& Sample) const override;
 	virtual FVector4		LightCenterPosition(const FVector4& ReceivingPosition, const FVector4& ReceivingNormal) const override;
+	virtual bool			AffectsBounds(const FBoxSphereBounds& Bounds) const override;
 	virtual bool			BehindSurface(const FVector4& TrianglePoint, const FVector4& TriangleNormal) const override;
 	virtual FVector4		GetDirectLightingDirection(const FVector4& Point, const FVector4& PointNormal) const override;
 
