Workaround for UE-63715

Strictly speaking directional lights outside of WORLD_MAX are not supported, but they happen to work before because their origin might be non zero. This workaround expands the radius of directional lights to FLT_MAX. However we are not supporting this so if this problem pops up again we will claim it as by design.

diff --git a/Engine/Source/Programs/UnrealLightmass/Private/ImportExport/LightmassScene.cpp b/Engine/Source/Programs/UnrealLightmass/Private/ImportExport/LightmassScene.cpp
--- a/Engine/Source/Programs/UnrealLightmass/Private/ImportExport/LightmassScene.cpp
+++ b/Engine/Source/Programs/UnrealLightmass/Private/ImportExport/LightmassScene.cpp
@@ -574,8 +574,8 @@ bool FLight::AffectsBounds(const FBoxSphereBounds& Bounds) const
 
 FSphere FLight::GetBoundingSphere() const
 {
-	// Directional lights will have a radius of WORLD_MAX
-	return FSphere(FVector(0, 0, 0), (float)WORLD_MAX);
+	// Directional lights will have a radius of FLT_MAX
+	return FSphere(FVector(0, 0, 0), (float)FLT_MAX);
 }
 
 /**
diff --git a/Engine/Source/Runtime/Engine/Classes/Components/LightComponent.h b/Engine/Source/Runtime/Engine/Classes/Components/LightComponent.h
--- a/Engine/Source/Runtime/Engine/Classes/Components/LightComponent.h
+++ b/Engine/Source/Runtime/Engine/Classes/Components/LightComponent.h
@@ -307,8 +307,8 @@ public:
 
 	virtual FSphere GetBoundingSphere() const
 	{
-		// Directional lights will have a radius of WORLD_MAX
-		return FSphere(FVector::ZeroVector, WORLD_MAX);
+		// Directional lights will have a radius of FLT_MAX
+		return FSphere(FVector::ZeroVector, FLT_MAX);
 	}
 
 	/**
diff --git a/Engine/Source/Runtime/Engine/Public/SceneManagement.h b/Engine/Source/Runtime/Engine/Public/SceneManagement.h
--- a/Engine/Source/Runtime/Engine/Public/SceneManagement.h
+++ b/Engine/Source/Runtime/Engine/Public/SceneManagement.h
@@ -950,8 +950,8 @@ public:
 
 	virtual FSphere GetBoundingSphere() const
 	{
-		// Directional lights will have a radius of WORLD_MAX
-		return FSphere(FVector::ZeroVector, WORLD_MAX);
+		// Directional lights will have a radius of FLT_MAX
+		return FSphere(FVector::ZeroVector, FLT_MAX);
 	}
 
 	/** @return radius of the light */
