Cherry-pic 7402145: Workaround for PolygonIrradiance() generating NaNs and got spreaded by other passes like final gathering.

diff --git a/Engine/Source/Programs/UnrealLightmass/Private/ImportExport/LightmassScene.cpp b/Engine/Source/Programs/UnrealLightmass/Private/ImportExport/LightmassScene.cpp
--- a/Engine/Source/Programs/UnrealLightmass/Private/ImportExport/LightmassScene.cpp
+++ b/Engine/Source/Programs/UnrealLightmass/Private/ImportExport/LightmassScene.cpp
@@ -1624,6 +1624,11 @@ FLinearColor FRectLight::GetDirectIntensity(const FVector4& Point, bool bCalcula
 	Poly[3] = ToLight - AxisX * Extent.X + AxisY * Extent.Y;
 
 	float DistanceAttenuation = PolygonIrradiance( Poly ).Size();
+
+	if (FMath::IsNaN(DistanceAttenuation))
+	{
+		DistanceAttenuation = 0;
+	}
 			
 	// TODO Move CPU cide
 	DistanceAttenuation /= 2 * Extent.X * Extent.Y;
