Redo cherry-pick 7402145: Workaround for PolygonIrradiance() generating NaNs and got spreaded by other passes like final gathering.

diff --git a/Engine/Source/Programs/UnrealLightmass/Private/ImportExport/LightmassScene.cpp b/Engine/Source/Programs/UnrealLightmass/Private/ImportExport/LightmassScene.cpp
--- a/Engine/Source/Programs/UnrealLightmass/Private/ImportExport/LightmassScene.cpp
+++ b/Engine/Source/Programs/UnrealLightmass/Private/ImportExport/LightmassScene.cpp
@@ -1625,7 +1625,7 @@ FLinearColor FRectLight::GetDirectIntensity(const FVector4& Point, bool bCalcula
 
 	float DistanceAttenuation = PolygonIrradiance( Poly ).Size();
 
-	if (FMath::IsNaN(DistanceAttenuation))
+	if (!FMath::IsFinite(DistanceAttenuation))
 	{
 		DistanceAttenuation = 0;
 	}
