More UE_ARRAY_COUNT fixes.

diff --git a/Engine/Source/Programs/UnrealLightmass/Private/LightmassCore/Math/LMkDOP.h b/Engine/Source/Programs/UnrealLightmass/Private/LightmassCore/Math/LMkDOP.h
--- a/Engine/Source/Programs/UnrealLightmass/Private/LightmassCore/Math/LMkDOP.h
+++ b/Engine/Source/Programs/UnrealLightmass/Private/LightmassCore/Math/LMkDOP.h
@@ -588,7 +588,7 @@ FORCEINLINE FOverlapInterval GetInterval(const FBox& Box, const FVector& Vector)
 	FOverlapInterval Result;
 	Result.Min = Result.Max = FVector::DotProduct(Vector, BoxVertices[0]);
 
-	for (int32 i = 1; i < ARRAY_COUNT(BoxVertices); ++i)
+	for (int32 i = 1; i < UE_ARRAY_COUNT(BoxVertices); ++i)
 	{
 		float Projection = FVector::DotProduct(Vector, BoxVertices[i]);
 		Result.Min = FMath::Min(Result.Min, Projection);
@@ -635,7 +635,7 @@ FORCEINLINE bool IntersectTriangleAndAABB(const FTriangle& Triangle, const FBox&
 		FVector::CrossProduct(BoxNormal2, TriangleEdge2)
 	};
 
-	for (int i = 0; i < ARRAY_COUNT(TestDirections); ++i)
+	for (int i = 0; i < UE_ARRAY_COUNT(TestDirections); ++i)
 	{
 		if (!OverlapOnAxis(Box, Triangle, TestDirections[i]))
 		{
diff --git a/Engine/Source/Programs/UnrealVersionSelector/Private/UnrealVersionSelector.cpp b/Engine/Source/Programs/UnrealVersionSelector/Private/UnrealVersionSelector.cpp
--- a/Engine/Source/Programs/UnrealVersionSelector/Private/UnrealVersionSelector.cpp
+++ b/Engine/Source/Programs/UnrealVersionSelector/Private/UnrealVersionSelector.cpp
@@ -246,7 +246,7 @@ bool TryGetEditorFileName(const FString& EngineDir, const FString& ProjectFileNa
 			virtual bool Visit(const TCHAR* FilenameOrDirectory, const FFileStatData& StatData)
 			{
 				static const TCHAR Extension[] = TEXT(".target");
-				static const int ExtensionLen = ARRAY_COUNT(Extension) - 1;
+				static const int ExtensionLen = UE_ARRAY_COUNT(Extension) - 1;
 
 				int Length = FCString::Strlen(FilenameOrDirectory);
 				if(Length >= ExtensionLen && FCString::Stricmp(FilenameOrDirectory + Length - ExtensionLen, Extension) == 0)
diff --git a/Engine/Source/Runtime/Engine/Private/Materials/MaterialInstance.cpp b/Engine/Source/Runtime/Engine/Private/Materials/MaterialInstance.cpp
--- a/Engine/Source/Runtime/Engine/Private/Materials/MaterialInstance.cpp
+++ b/Engine/Source/Runtime/Engine/Private/Materials/MaterialInstance.cpp
@@ -1339,7 +1339,7 @@ void UMaterialInstance::ValidateTextureOverrides(ERHIFeatureLevel::Type InFeatur
 	FString MaterialName;
 	GetName(MaterialName);
 
-	for (int32 TypeIndex = 0; TypeIndex < ARRAY_COUNT(TextureExpressions); ++TypeIndex)
+	for (int32 TypeIndex = 0; TypeIndex < UE_ARRAY_COUNT(TextureExpressions); ++TypeIndex)
 	{
 		for (FMaterialUniformExpressionTexture* TextureExpression : *TextureExpressions[TypeIndex])
 		{
