Remove unused code causing deprecation warning
#rnx

diff --git a/Engine/Source/Programs/UnrealLightmass/Private/Lighting/AdaptiveVolumetricLightmap.cpp b/Engine/Source/Programs/UnrealLightmass/Private/Lighting/AdaptiveVolumetricLightmap.cpp
--- a/Engine/Source/Programs/UnrealLightmass/Private/Lighting/AdaptiveVolumetricLightmap.cpp
+++ b/Engine/Source/Programs/UnrealLightmass/Private/Lighting/AdaptiveVolumetricLightmap.cpp
@@ -150,15 +150,6 @@ bool FStaticLightingSystem::DoesVoxelIntersectSceneGeometry(const FBox& CellBoun
 
 			if (MeshInstance->StaticMesh->VoxelizationMesh != nullptr)
 			{
-				auto TransformBox = [](const FBox& Box, const FMatrix& Transform) -> const FBox {
-					FVector TransformedCorners[8];
-					for (int32 Corner = 0; Corner < 8; Corner++)
-					{
-						TransformedCorners[Corner] = Transform.TransformPosition(Box.GetExtrema(Corner));
-					}
-					return FBox(TransformedCorners, UE_ARRAY_COUNT(TransformedCorners));
-				};
-
 				if (MeshInstance->Mapping->GetVolumeMapping() == nullptr)
 				{
 					if (MeshInstance->StaticMesh->VoxelizationMesh->IntersectBox(ExpandedCellBoundsSurfaceGeometry.TransformBy(MeshInstance->WorldToLocal)) != nullptr)
