Remove unused code
#5749

diff --git a/Engine/Source/Programs/UnrealLightmass/Private/Lighting/AdaptiveVolumetricLightmap.cpp b/Engine/Source/Programs/UnrealLightmass/Private/Lighting/AdaptiveVolumetricLightmap.cpp
--- a/Engine/Source/Programs/UnrealLightmass/Private/Lighting/AdaptiveVolumetricLightmap.cpp
+++ b/Engine/Source/Programs/UnrealLightmass/Private/Lighting/AdaptiveVolumetricLightmap.cpp
@@ -202,11 +202,6 @@ bool FStaticLightingSystem::DoesVoxelIntersectSceneGeometry(const FBox& CellBoun
 
 bool FStaticLightingSystem::ShouldRefineVoxel(int32 TreeDepth, const FBox& CellBounds, const TArray<FVector>& VoxelTestPositions, bool bDebugThisVoxel) const
 {
-	if (bDebugThisVoxel)
-	{
-		int32 asdf = 0;
-	}
-
 	const bool bCellInsideImportanceVolume = Scene.IsBoxInImportanceVolume(CellBounds);
 
 	// The volumetric lightmap bounds are larger than the importance volume bounds, since we force the volumetric lightmap volume to have cube voxels
@@ -232,11 +227,6 @@ bool FStaticLightingSystem::ShouldRefineVoxel(int32 TreeDepth, const FBox& CellB
 		}
 	}
 
-	if (bDebugThisVoxel)
-	{
-		int32 asdf = 0;
-	}
-
 	// Default unrestricted
 	FIntPoint AllowedMipRange(0, INT_MAX);
 
@@ -345,11 +335,6 @@ bool FStaticLightingSystem::ShouldRefineVoxel(int32 TreeDepth, const FBox& CellB
 						{
 							if (PointUnderTriangle(TestPositions[PointIndex], Triangle))
 							{
-								if (bDebugThisVoxel)
-								{
-									int32 asdf = 0;
-								}
-
 								PositionUnderLandscape[PointIndex] = true;
 							}
 						}
@@ -358,11 +343,6 @@ bool FStaticLightingSystem::ShouldRefineVoxel(int32 TreeDepth, const FBox& CellB
 			}
 		}
 
-		if (bDebugThisVoxel)
-		{
-			int32 asdf = 0;
-		}
-
 		bool bAllPointsUnderLandscape = true;
 
 		for (int32 PointIndex = 0; PointIndex < PositionUnderLandscape.Num(); PointIndex++)
@@ -439,12 +419,6 @@ void FStaticLightingSystem::RecursivelyBuildBrickTree(
 						const FBox CellBounds(ChildCellPosition, ChildCellPosition + WorldChildCellSize);
 
 						const bool bChildCoveringDebugPosition = bDebugVolumetricLightmapCell && CellBounds.IsInside(DebugWorldPosition);
-
-						if (bChildCoveringDebugPosition)
-						{
-							int32 asdf = 0;
-						}
-
 						const bool bSubdivideCell = ShouldRefineVoxel(TreeDepth + 1, CellBounds, VoxelTestPositions, bChildCoveringDebugPosition);
 
 						if (bSubdivideCell)
@@ -511,11 +485,6 @@ void FStaticLightingSystem::ProcessVolumetricLightmapBrickTask(FVolumetricLightm
 	const FIrradianceBrickBuildData& BuildData = Task->BuildData;
 	FIrradianceBrickData& BrickData = Task->BrickData;
 
-	if (BuildData.bDebugBrick)
-	{
-		int32 asdf = 0;
-	}
-
 	const int32 BrickSize = VolumetricLightmapSettings.BrickSize;
 	const int32 BrickSizeLog2 = FMath::FloorLog2(BrickSize);
 	const int32 DetailCellsPerTopLevelBrick = 1 << (VolumetricLightmapSettings.MaxRefinementLevels * BrickSizeLog2);
@@ -603,7 +572,6 @@ void FStaticLightingSystem::ProcessVolumetricLightmapBrickTask(FVolumetricLightm
 				if (bDebugSamples)
 				{
 					Task->MappingContext.DebugOutput->bValid = true;
-					int32 asdf = 0;
 				}
 
 				float BackfacingHitsFraction = 0.0f;
@@ -661,11 +629,6 @@ void FStaticLightingSystem::ProcessVolumetricLightmapBrickTask(FVolumetricLightm
 	BrickData.AverageClosestGeometryDistance = AverageClosestGeometryDistance / TotalBrickSize;
 	
 	const bool bCullBrick = bAllCellsInsideGeometry;
-	
-	if (BuildData.bDebugBrick)
-	{
-		int32 asdf = 0;
-	}
 
 	if (bCullBrick && BuildData.TreeDepth > 0 && !BuildData.bHasChildren)
 	{
@@ -767,11 +730,6 @@ void FStaticLightingSystem::CalculateAdaptiveVolumetricLightmap(int32 TaskIndex)
 	TArray<FVector> VoxelTestPositions;
 	GenerateVoxelTestPositions(VoxelTestPositions);
 
-	if (bCoveringDebugPosition)
-	{
-		int32 asdf = 0;
-	}
-
 	TArray<FIrradianceBrickBuildData> BrickBuildData;
 	RecursivelyBuildBrickTree(StartCellIndex, NumCells, FIntVector::ZeroValue, 0, bCoveringDebugPosition, TopLevelBounds, VoxelTestPositions, BrickBuildData);
 
diff --git a/Engine/Source/Programs/UnrealLightmass/Private/Lighting/FinalGather.cpp b/Engine/Source/Programs/UnrealLightmass/Private/Lighting/FinalGather.cpp
--- a/Engine/Source/Programs/UnrealLightmass/Private/Lighting/FinalGather.cpp
+++ b/Engine/Source/Programs/UnrealLightmass/Private/Lighting/FinalGather.cpp
@@ -82,11 +82,6 @@ void FStaticLightingSystem::CalculateVolumeSampleIncidentRadiance(
 	bool bDebugThisSample
 	) const
 {
-	if (bDebugThisSample)
-	{
-		int32 asdf = 0;
-	}
-
 	const double StartTime = FPlatformTime::Seconds();
 
 	const FVector4 Position = LightingSample.GetPosition();
@@ -408,13 +403,6 @@ FLinearColor FStaticLightingSystem::FinalGatherSample(
 	FLinearColor Lighting = FLinearColor::Black;
 	OutStationarySkyLighting = FLinearColor::Black;
 
-#if ALLOW_LIGHTMAP_SAMPLE_DEBUGGING
-	if (bDebugThisTexel)
-	{
-		int32 asdf = 0;
-	}
-#endif
-
 	bool bPositiveSample = false;
 
 	OutUnoccludedSkyVector = RayIntersection.bIntersects ? FVector(0) : FVector(WorldPathDirection);
@@ -692,11 +680,6 @@ public:
 		// This is necessary for the neighbor comparisons to work right
 		for (int32 RefinementDepth = 0; RefinementDepth < NumAdaptiveRefinementLevels; RefinementDepth++)
 		{
-			if (bDebugThisTexel)
-			{
-				int32 asdf = 0;
-			}
-
 			FLinearColor TotalLighting = FLinearColor::Black;
 
 			// Recalculate total lighting based on the refined results
@@ -959,11 +942,6 @@ public:
 			// Swap input and output for the next step
 			Swap(CurrentNodesToRefine, NextNodesToRefine);
 
-			if (bDebugThisTexel)
-			{
-				int32 asdf = 0;
-			}
-
 			FVector4 WorldPathDirections[4];
 			FVector4 TangentPathDirections[4];
 			FLightRay LightRays[4];
diff --git a/Engine/Source/Programs/UnrealLightmass/Private/Lighting/LightingCache.h b/Engine/Source/Programs/UnrealLightmass/Private/Lighting/LightingCache.h
--- a/Engine/Source/Programs/UnrealLightmass/Private/Lighting/LightingCache.h
+++ b/Engine/Source/Programs/UnrealLightmass/Private/Lighting/LightingCache.h
@@ -383,11 +383,6 @@ bool TLightingCache<SampleType>::InterpolateLighting(
 	{
 		const FRecord<SampleType>& LightingRecord = OctreeIt.GetCurrentElement();
 
-		if (bDebugThisSample)
-		{
-			int asdf = 0;
-		}
-
 		// Check whether the query point is farther than the record's intersection distance for the direction to the query point.
 		const float DistanceSquared = (LightingRecord.Vertex.WorldPosition - Vertex.WorldPosition).SizeSquared3();
 		if (DistanceSquared > FMath::Square(LightingRecord.BoundingRadius))
diff --git a/Engine/Source/Programs/UnrealLightmass/Private/Lighting/LightingSystem.cpp b/Engine/Source/Programs/UnrealLightmass/Private/Lighting/LightingSystem.cpp
--- a/Engine/Source/Programs/UnrealLightmass/Private/Lighting/LightingSystem.cpp
+++ b/Engine/Source/Programs/UnrealLightmass/Private/Lighting/LightingSystem.cpp
@@ -59,11 +59,6 @@ static void ConvertToLightSampleHelper(const FGatheredLightSample& InGatheredLig
 
 FLightSample FGatheredLightMapSample::ConvertToLightSample(bool bDebugThisSample) const
 {
-	if (bDebugThisSample)
-	{
-		int32 asdf = 0;
-	}
-
 	FLightSample NewSample;
 	NewSample.bIsMapped = bIsMapped;
 
diff --git a/Engine/Source/Programs/UnrealLightmass/Private/Lighting/LightingSystem.inl b/Engine/Source/Programs/UnrealLightmass/Private/Lighting/LightingSystem.inl
--- a/Engine/Source/Programs/UnrealLightmass/Private/Lighting/LightingSystem.inl
+++ b/Engine/Source/Programs/UnrealLightmass/Private/Lighting/LightingSystem.inl
@@ -103,11 +103,6 @@ void FStaticLightingSystem::CalculateApproximateDirectLighting(
 	TGatheredLightSample<SHOrder>& OutToggleableDirectLighting,
 	float& OutToggleableDirectionalLightShadowing) const
 {
-	if (bDebugThisSample)
-	{
-		int asdf = 0;
-	}
-
 	check(VertexOffsets.Num() > 0);
 
 	for (int32 LightIndex = 0; LightIndex < Lights.Num(); LightIndex++)
@@ -193,11 +188,6 @@ void FStaticLightingSystem::CalculateApproximateDirectLighting(
 				Transmission = FLinearColor::White;
 			}
 
-			if (bDebugThisSample)
-			{
-				int asdf = 0;
-			}
-
 			{
 				// Calculate the direction from the vertex to the light.
 				const FVector4 WorldLightVector = Light->GetDirectLightingDirection(Vertex.WorldPosition, Vertex.WorldTangentZ);
diff --git a/Engine/Source/Programs/UnrealLightmass/Private/Lighting/TextureMapping.cpp b/Engine/Source/Programs/UnrealLightmass/Private/Lighting/TextureMapping.cpp
--- a/Engine/Source/Programs/UnrealLightmass/Private/Lighting/TextureMapping.cpp
+++ b/Engine/Source/Programs/UnrealLightmass/Private/Lighting/TextureMapping.cpp
@@ -664,11 +664,6 @@ void FStaticLightingSystem::ProcessTextureMapping(FStaticLightingTextureMapping*
 		// Release corner information as it is no longer needed
 		TexelToCornersMap.Empty();
 
-		if (bDebugThisMapping)
-		{
-			int32 asdf = 0;
-		}
-
 		// Calculate direct lighting using the direct photon map.
 		// This is only useful for debugging what the final gather rays see.
 		if (bCalculateDirectLightingFromPhotons)
@@ -3176,11 +3171,6 @@ void FStaticLightingSystem::CalculateIndirectLightingTextureMapping(
 		
 		if (IrradianceCachingSettings.bAllowIrradianceCaching)
 		{
-			if (bDebugThisMapping)
-			{
-				int32 asdf = 0;
-			}
-
 			const int32 InterpolationTaskSize = IrradianceCachingSettings.InterpolateTaskSize;
 			int32 NumIILTasksSubmitted = 0;
 
@@ -3224,11 +3214,6 @@ void FStaticLightingSystem::CalculateIndirectLightingTextureMapping(
 			} 
 			while (TextureMapping->NumOutstandingInterpolationTasks > 0);
 
-			if (bDebugThisMapping)
-			{
-				int32 asdf = 0;
-			}
-
 			TArray<FInterpolateIndirectTaskDescription*> CompletedTasks;
 			TextureMapping->CompletedInterpolationTasks.PopAll(CompletedTasks);
 			check(CompletedTasks.Num() == NumIILTasksSubmitted);
