Multiple small improvements to Lightmass

* Improve the back face detection of our 'representative surfel' method to better mitigate dark corners
* Force threads to process local tasks beforce fetching new texture mapping tasks to improve work stealing

diff --git a/Engine/Source/Programs/UnrealLightmass/Private/Lighting/AdaptiveVolumetricLightmap.cpp b/Engine/Source/Programs/UnrealLightmass/Private/Lighting/AdaptiveVolumetricLightmap.cpp
--- a/Engine/Source/Programs/UnrealLightmass/Private/Lighting/AdaptiveVolumetricLightmap.cpp
+++ b/Engine/Source/Programs/UnrealLightmass/Private/Lighting/AdaptiveVolumetricLightmap.cpp
@@ -638,9 +638,7 @@ void FStaticLightingSystem::ProcessVolumetricLightmapBrickTask(FVolumetricLightm
 
 void FStaticLightingSystem::ProcessVolumetricLightmapTaskIfAvailable()
 {
-	FVolumetricLightmapBrickTaskDescription* NextTask = VolumetricLightmapBrickTasks.Pop();
-
-	if (NextTask)
+	while (FVolumetricLightmapBrickTaskDescription * NextTask = VolumetricLightmapBrickTasks.Pop())
 	{
 		//UE_LOG(LogLightmass, Warning, TEXT("Thread picked up volumetric lightmap task"));
 		ProcessVolumetricLightmapBrickTask(NextTask);
diff --git a/Engine/Source/Programs/UnrealLightmass/Private/Lighting/LightingSystem.cpp b/Engine/Source/Programs/UnrealLightmass/Private/Lighting/LightingSystem.cpp
--- a/Engine/Source/Programs/UnrealLightmass/Private/Lighting/LightingSystem.cpp
+++ b/Engine/Source/Programs/UnrealLightmass/Private/Lighting/LightingSystem.cpp
@@ -1585,7 +1585,25 @@ void FStaticLightingSystem::ThreadLoop(bool bIsMainThread, int32 ThreadIndex, FT
 	while (!bIsDone)
 	{
 		const double StartLoopTime = FPlatformTime::Seconds();
-		
+
+		while (FCacheIndirectTaskDescription * NextCacheTask = CacheIndirectLightingTasks.Pop())
+		{
+			//UE_LOG(LogLightmass, Warning, TEXT("Thread %u picked up Cache Indirect task for %u"), ThreadIndex, NextCacheTask->TextureMapping->Guid.D);
+			ProcessCacheIndirectLightingTask(NextCacheTask, false);
+			NextCacheTask->TextureMapping->CompletedCacheIndirectLightingTasks.Push(NextCacheTask);
+			FPlatformAtomics::InterlockedDecrement(&NextCacheTask->TextureMapping->NumOutstandingCacheTasks);
+		}
+
+		while (FInterpolateIndirectTaskDescription * NextInterpolateTask = InterpolateIndirectLightingTasks.Pop())
+		{
+			//UE_LOG(LogLightmass, Warning, TEXT("Thread %u picked up Interpolate indirect task for %u"), ThreadIndex, NextInterpolateTask->TextureMapping->Guid.D);
+			ProcessInterpolateTask(NextInterpolateTask, false);
+			NextInterpolateTask->TextureMapping->CompletedInterpolationTasks.Push(NextInterpolateTask);
+			FPlatformAtomics::InterlockedDecrement(&NextInterpolateTask->TextureMapping->NumOutstandingInterpolationTasks);
+		}
+
+		ProcessVolumetricLightmapTaskIfAvailable();
+
 		if (NumOutstandingVolumeDataLayers > 0)
 		{
 			const int32 ThreadZ = FPlatformAtomics::InterlockedIncrement(&OutstandingVolumeDataLayerIndex);
@@ -1677,28 +1695,6 @@ void FStaticLightingSystem::ThreadLoop(bool bIsMainThread, int32 ThreadIndex, FT
 				bSignaledMappingsComplete = true;
 				GSwarm->SendMessage( NSwarm::FTimingMessage( NSwarm::PROGSTATE_Processing0, ThreadIndex ) );
 			}
-
-			FCacheIndirectTaskDescription* NextCacheTask = CacheIndirectLightingTasks.Pop();
-
-			if (NextCacheTask)
-			{
-				//UE_LOG(LogLightmass, Warning, TEXT("Thread %u picked up Cache Indirect task for %u"), ThreadIndex, NextCacheTask->TextureMapping->Guid.D);
-				ProcessCacheIndirectLightingTask(NextCacheTask, false);
-				NextCacheTask->TextureMapping->CompletedCacheIndirectLightingTasks.Push(NextCacheTask);
-				FPlatformAtomics::InterlockedDecrement(&NextCacheTask->TextureMapping->NumOutstandingCacheTasks);
-			}
-
-			FInterpolateIndirectTaskDescription* NextInterpolateTask = InterpolateIndirectLightingTasks.Pop();
-
-			if (NextInterpolateTask)
-			{
-				//UE_LOG(LogLightmass, Warning, TEXT("Thread %u picked up Interpolate indirect task for %u"), ThreadIndex, NextInterpolateTask->TextureMapping->Guid.D);
-				ProcessInterpolateTask(NextInterpolateTask, false);
-				NextInterpolateTask->TextureMapping->CompletedInterpolationTasks.Push(NextInterpolateTask);
-				FPlatformAtomics::InterlockedDecrement(&NextInterpolateTask->TextureMapping->NumOutstandingInterpolationTasks);
-			}
-
-			ProcessVolumetricLightmapTaskIfAvailable();
 			
 			if (NumVolumeSampleTasksOutstanding > 0)
 			{
@@ -1716,10 +1712,7 @@ void FStaticLightingSystem::ThreadLoop(bool bIsMainThread, int32 ThreadIndex, FT
 				}
 			}
 
-			if (!NextCacheTask 
-				&& !NextInterpolateTask
-				&& NumVolumeSampleTasksOutstanding <= 0
-				&& NumOutstandingVolumeDataLayers <= 0)
+			if (NumVolumeSampleTasksOutstanding <= 0 && NumOutstandingVolumeDataLayers <= 0)
 			{
 				if (TasksInProgressThatWillNeedHelp <= 0 && !bRequestForTaskTimedOut)
 				{
diff --git a/Engine/Source/Programs/UnrealLightmass/Private/Lighting/LightingSystem.h b/Engine/Source/Programs/UnrealLightmass/Private/Lighting/LightingSystem.h
--- a/Engine/Source/Programs/UnrealLightmass/Private/Lighting/LightingSystem.h
+++ b/Engine/Source/Programs/UnrealLightmass/Private/Lighting/LightingSystem.h
@@ -2058,6 +2058,8 @@ private:
 
 	void RasterizeToSurfaceCacheTextureMapping(FStaticLightingTextureMapping* TextureMapping, bool bDebugThisMapping, FTexelToVertexMap& TexelToVertexMap);
 
+	void AdjustRepresentativeSurfelForTexelsTextureMapping(FStaticLightingTextureMapping* TextureMapping, FTexelToVertexMap& TexelToVertexMap, FTexelToCornersMap& TexelToCornersMap, FGatheredLightMapData2D* LightMapData, FStaticLightingMappingContext& MappingContext,	bool bDebugThisMapping) const;
+
 	void RadiosityIterationTextureMapping(FStaticLightingTextureMapping* TextureMapping, int32 PassIndex);
 
 	void RadiosityIterationCachedHitpointsTextureMapping(const FTexelToVertexMap& TexelToVertexMap, FStaticLightingTextureMapping* TextureMapping, int32 PassIndex);
diff --git a/Engine/Source/Programs/UnrealLightmass/Private/Lighting/TextureMapping.cpp b/Engine/Source/Programs/UnrealLightmass/Private/Lighting/TextureMapping.cpp
--- a/Engine/Source/Programs/UnrealLightmass/Private/Lighting/TextureMapping.cpp
+++ b/Engine/Source/Programs/UnrealLightmass/Private/Lighting/TextureMapping.cpp
@@ -267,6 +267,11 @@
 
 void FStaticLightingSystem::RasterizeToSurfaceCacheTextureMapping(FStaticLightingTextureMapping* TextureMapping, bool bDebugThisMapping, FTexelToVertexMap& TexelToVertexMap)
 {
+	FStaticLightingMappingContext MappingContext(TextureMapping->Mesh, *this, &StartupDebugOutput);
+
+	FTexelToCornersMap TexelToCornersMap(TextureMapping->SurfaceCacheSizeX, TextureMapping->SurfaceCacheSizeY);
+	CalculateTexelCorners(TextureMapping->Mesh, TexelToCornersMap, TextureMapping->LightmapTextureCoordinateIndex, bDebugThisMapping);
+
 	// Using conservative rasterization, which uses super sampling to try to detect all texels that should be mapped.
 	for(int32 TriangleIndex = 0;TriangleIndex < TextureMapping->Mesh->NumTriangles;TriangleIndex++)
 	{
@@ -327,30 +332,112 @@
 		}
 	}
 
-	for (int32 Y = 0; Y < TextureMapping->SurfaceCacheSizeY; Y++)
+	AdjustRepresentativeSurfelForTexelsTextureMapping(TextureMapping, TexelToVertexMap, TexelToCornersMap, nullptr, MappingContext, bDebugThisMapping);
+}
+
+void FStaticLightingSystem::AdjustRepresentativeSurfelForTexelsTextureMapping(
+	FStaticLightingTextureMapping* TextureMapping, 
+	FTexelToVertexMap& TexelToVertexMap, 
+	FTexelToCornersMap& TexelToCornersMap,
+	FGatheredLightMapData2D* LightMapData,
+	FStaticLightingMappingContext& MappingContext,
+	bool bDebugThisMapping) const
+{
+	// Iterate over each texel and normalize vectors, calculate texel radius
+	for (int32 Y = 0; Y < TexelToVertexMap.GetSizeY(); Y++)
 	{
-		for (int32 X = 0; X < TextureMapping->SurfaceCacheSizeX; X++)
+		for (int32 X = 0; X < TexelToVertexMap.GetSizeX(); X++)
 		{
+			bool bDebugThisTexel = false;
+#if ALLOW_LIGHTMAP_SAMPLE_DEBUGGING
+			if (bDebugThisMapping
+				&& Y == Scene.DebugInput.LocalY
+				&& X == Scene.DebugInput.LocalX)
+			{
+				bDebugThisTexel = true;
+			}
+#endif
+			bool bFoundValidCorner = false;
+			const FTexelToCorners& TexelToCorners = TexelToCornersMap(X, Y);
+			for (int32 CornerIndex = 0; CornerIndex < NumTexelCorners; CornerIndex++)
+			{
+				bFoundValidCorner = bFoundValidCorner || TexelToCorners.bValid[CornerIndex];
+			}
+
 			FTexelToVertex& TexelToVertex = TexelToVertexMap(X, Y);
-			if (TexelToVertex.TotalSampleWeight > 0.0f)
+			if (TexelToVertex.TotalSampleWeight > 0.0f || bFoundValidCorner)
 			{
-				if (GeneralSettings.bUseMaxWeight)
+				// Use a corner if none of the other samples were valid
+				if (TexelToVertex.TotalSampleWeight < DELTA)
 				{
+					for (int32 CornerIndex = 0; CornerIndex < NumTexelCorners; CornerIndex++)
+					{
+						if (TexelToCorners.bValid[CornerIndex])
+						{
+							TexelToVertex.TotalSampleWeight = 1.0f;
+							TexelToVertex.WorldPosition = TexelToCorners.Corners[CornerIndex].WorldPosition;
+							TexelToVertex.WorldTangentX = TexelToCorners.WorldTangentX;
+							TexelToVertex.WorldTangentY = TexelToCorners.WorldTangentY;
+							TexelToVertex.WorldTangentZ = TexelToCorners.WorldTangentZ;
+							TexelToVertex.TriangleNormal = TexelToCorners.WorldTangentZ;
+							break;
+						}
+					}
+				}
+				else if (GeneralSettings.bUseMaxWeight)
+				{
 					// Weighted average
 					TexelToVertex.WorldTangentX = TexelToVertex.WorldTangentX / TexelToVertex.TotalSampleWeight;
 					TexelToVertex.WorldTangentY = TexelToVertex.WorldTangentY / TexelToVertex.TotalSampleWeight;
 					TexelToVertex.WorldTangentZ = TexelToVertex.WorldTangentZ / TexelToVertex.TotalSampleWeight;
 					TexelToVertex.TriangleNormal = TexelToVertex.TriangleNormal / TexelToVertex.TotalSampleWeight;
+
+					// Weighted average of opposing vectors can result in a zero vector, fixup with corner
+					if (bFoundValidCorner
+						&& (TexelToVertex.WorldTangentX.SizeSquared3() < KINDA_SMALL_NUMBER
+							|| TexelToVertex.WorldTangentZ.SizeSquared3() < KINDA_SMALL_NUMBER
+							|| TexelToVertex.TriangleNormal.SizeSquared3() < KINDA_SMALL_NUMBER))
+					{
+						for (int32 CornerIndex = 0; CornerIndex < NumTexelCorners; CornerIndex++)
+						{
+							if (TexelToCorners.bValid[CornerIndex])
+							{
+								TexelToVertex.WorldTangentX = TexelToCorners.WorldTangentX;
+								TexelToVertex.WorldTangentY = TexelToCorners.WorldTangentY;
+								TexelToVertex.WorldTangentZ = TexelToCorners.WorldTangentZ;
+								TexelToVertex.TriangleNormal = TexelToCorners.WorldTangentZ;
+								break;
+							}
+						}
+					}
 				}
 
+				if (LightMapData != nullptr)
+				{
+					// Mark the texel as mapped to some geometry in the scene
+					FGatheredLightMapSample& CurrentLightSample = (*LightMapData)(X, Y);
+					CurrentLightSample.bIsMapped = true;
+				}
+
+				if (MaterialSettings.bUseNormalMapsForLighting && TextureMapping->Mesh->HasImportedNormal(TexelToVertex.ElementIndex))
+				{
+					const FVector4 TangentNormal = TextureMapping->Mesh->EvaluateNormal(TexelToVertex.TextureCoordinates[0], TexelToVertex.ElementIndex);
+
+					const FVector4 WorldTangentRow0(TexelToVertex.WorldTangentX.X, TexelToVertex.WorldTangentY.X, TexelToVertex.WorldTangentZ.X);
+					const FVector4 WorldTangentRow1(TexelToVertex.WorldTangentX.Y, TexelToVertex.WorldTangentY.Y, TexelToVertex.WorldTangentZ.Y);
+					const FVector4 WorldTangentRow2(TexelToVertex.WorldTangentX.Z, TexelToVertex.WorldTangentY.Z, TexelToVertex.WorldTangentZ.Z);
+					const FVector4 WorldVector(
+						Dot3(WorldTangentRow0, TangentNormal),
+						Dot3(WorldTangentRow1, TangentNormal),
+						Dot3(WorldTangentRow2, TangentNormal)
+					);
+
+					TexelToVertex.WorldTangentZ = WorldVector;
+				}
+
 				// Normalize the tangent basis and ensure it is orthonormal
-				TexelToVertex.WorldTangentZ = TexelToVertex.WorldTangentZ.GetSafeNormal();
+				TexelToVertex.WorldTangentZ = TexelToVertex.WorldTangentZ.GetUnsafeNormal3();
 
-				if (TexelToVertex.TriangleNormal.IsNearlyZero3())
-				{
-					TexelToVertex.TriangleNormal = TexelToVertex.WorldTangentZ;
-				}
-
 				const bool bUseVertexNormalForHemisphereGather = TextureMapping->Mesh->UseVertexNormalForHemisphereGather(TexelToVertex.ElementIndex);
 				TexelToVertex.TriangleNormal = bUseVertexNormalForHemisphereGather ? TexelToVertex.WorldTangentZ : TexelToVertex.TriangleNormal.GetUnsafeNormal3();
 				checkSlow(!TexelToVertex.TriangleNormal.ContainsNaN());
@@ -376,6 +463,132 @@
 				checkSlow(Dot3(TexelToVertex.WorldTangentZ, TexelToVertex.WorldTangentY) < KINDA_SMALL_NUMBER);
 				checkSlow(Dot3(TexelToVertex.WorldTangentX, TexelToVertex.WorldTangentY) < KINDA_SMALL_NUMBER);
 				checkSlow(Dot3(TexelToVertex.WorldTangentX, TexelToVertex.WorldTangentZ) < KINDA_SMALL_NUMBER);
+
+				// Calculate the bounding radius of the texel
+				// Use the closest corner as it's likely that's on the same section of a split texel 
+				// (A texel shared by multiple UV charts that has sub samples on triangles in different smoothing groups)
+				float MinDistanceSquared = FLT_MAX;
+				if (bFoundValidCorner)
+				{
+					for (int32 CornerIndex = 0; CornerIndex < NumTexelCorners; CornerIndex++)
+					{
+						if (TexelToCorners.bValid[CornerIndex])
+						{
+							const float CornerDistSquared = (TexelToCorners.Corners[CornerIndex].WorldPosition - TexelToVertex.WorldPosition).SizeSquared3();
+							if (CornerDistSquared < MinDistanceSquared)
+							{
+								MinDistanceSquared = CornerDistSquared;
+							}
+						}
+					}
+				}
+				else
+				{
+					MinDistanceSquared = SceneConstants.SmallestTexelRadius;
+				}
+				TexelToVertex.TexelRadius = FMath::Max(FMath::Sqrt(MinDistanceSquared), SceneConstants.SmallestTexelRadius);
+				MappingContext.Stats.NumMappedTexels++;
+
+				{
+					const FFullStaticLightingVertex FullVertex = TexelToVertex.GetFullVertex();
+					const FVector4 TexelCenterOffset = FullVertex.WorldPosition + FullVertex.TriangleNormal * TexelToVertex.TexelRadius * SceneConstants.VisibilityNormalOffsetSampleRadiusScale;
+
+					FLightRayIntersection Intersections[4];
+					bool bHitBackfaces[4];
+
+					FVector2D CornerSigns[4];
+					CornerSigns[0] = FVector2D(1, 1);
+					CornerSigns[1] = FVector2D(-1, 1);
+					CornerSigns[2] = FVector2D(1, -1);
+					CornerSigns[3] = FVector2D(-1, -1);
+
+					float BackfaceDetectionDistance = TexelToVertex.TexelRadius * 2.5f;
+
+					for (int32 CornerIndex = 0; CornerIndex < ARRAY_COUNT(CornerSigns); CornerIndex++)
+					{
+						TraceToTexelCorner(
+							TexelCenterOffset,
+							FullVertex,
+							CornerSigns[CornerIndex],
+							// Note: Searching the entire influence of the texel after interpolation, which is 2x the sample radius
+							BackfaceDetectionDistance,
+							MappingContext,
+							Intersections[CornerIndex],
+							bHitBackfaces[CornerIndex],
+							bDebugThisTexel);
+					}
+
+					int32 ClosestIntersectionIndex = INDEX_NONE;
+					float ClosestIntersectionDistanceSq = FLT_MAX;
+
+					int32 ClosestBackfacingIntersectionIndex = INDEX_NONE;
+					// Limit the distance that we will search for an intersecting backface in order to move the shading position to the texel radius
+					float ClosestBackfacingIntersectionDistanceSq = BackfaceDetectionDistance * BackfaceDetectionDistance;
+
+					for (int32 CornerIndex = 0; CornerIndex < ARRAY_COUNT(CornerSigns); CornerIndex++)
+					{
+						if (Intersections[CornerIndex].bIntersects)
+						{
+							const float DistanceSquared = (Intersections[CornerIndex].IntersectionVertex.WorldPosition - TexelCenterOffset).SizeSquared3();
+
+							if (ClosestIntersectionIndex == INDEX_NONE || DistanceSquared < ClosestIntersectionDistanceSq)
+							{
+								ClosestIntersectionDistanceSq = DistanceSquared;
+								ClosestIntersectionIndex = CornerIndex;
+							}
+
+							if (bHitBackfaces[CornerIndex] && DistanceSquared < ClosestBackfacingIntersectionDistanceSq)
+							{
+								ClosestBackfacingIntersectionDistanceSq = DistanceSquared;
+								ClosestBackfacingIntersectionIndex = CornerIndex;
+							}
+						}
+					}
+
+					if (ClosestIntersectionIndex != INDEX_NONE)
+					{
+						checkSlow(Intersections[ClosestIntersectionIndex].bIntersects);
+
+						// Mark the texel as intersecting another surface so we can avoid filtering across it later
+						TexelToVertex.bIntersectingSurface = true;
+
+						TexelToVertex.TexelRadius = FMath::Min(TexelToVertex.TexelRadius, FMath::Sqrt(ClosestIntersectionDistanceSq));
+					}
+
+					// Give preference to moving the shading position outside of backfaces
+					int32 IntersectionIndexForShadingPositionMovement = ClosestBackfacingIntersectionIndex;
+
+					// Note: this is disabled as it causes problems in cracks, the lighting position will be moved inside the object
+					/*
+					// Even if we didn't hit any backfaces, still move the shading position away from an intersecting frontface if it is close enough
+					if (IntersectionIndexForShadingPositionMovement == INDEX_NONE
+						&& ClosestIntersectionDistanceSq < (TexelToVertex.TexelRadius / 2) * (TexelToVertex.TexelRadius / 2))
+					{
+						IntersectionIndexForShadingPositionMovement = ClosestIntersectionIndex;
+					}*/
+
+					if (IntersectionIndexForShadingPositionMovement != INDEX_NONE)
+					{
+						// Move the shading position outside the surface that is intersecting this texel
+						const FVector4 OffsetShadingPosition = Intersections[IntersectionIndexForShadingPositionMovement].IntersectionVertex.WorldPosition
+							// Move along the intersecting surface's normal but also away from the texel a bit to prevent incorrect self occlusion
+							+ (Intersections[IntersectionIndexForShadingPositionMovement].IntersectionVertex.WorldTangentZ + TexelToVertex.TriangleNormal) * 0.5f * TexelToVertex.TexelRadius * SceneConstants.VisibilityNormalOffsetSampleRadiusScale;
+
+						// Project back onto plane of texel to avoid incorrect self occlusion
+						TexelToVertex.WorldPosition = OffsetShadingPosition + TexelToVertex.TriangleNormal * Dot3(TexelToVertex.TriangleNormal, TexelToVertex.WorldPosition - OffsetShadingPosition);
+
+						TexelToVertex.TexelRadius = (OffsetShadingPosition - Intersections[IntersectionIndexForShadingPositionMovement].IntersectionVertex.WorldPosition).Size3();
+					}
+				}
 			}
+			else
+			{
+				if (LightMapData != nullptr)
+				{
+					// Mark unmapped texels with the supplied 'UnmappedTexelColor'.
+					FGatheredLightMapSample& CurrentLightSample = (*LightMapData)(X, Y);
+					CurrentLightSample.AddWeighted(FGatheredLightSampleUtil::AmbientLight<2>(Scene.GeneralSettings.UnmappedTexelColor), 1.0f);
+				}
+			}
 		}
 	}
@@ -995,243 +1208,7 @@
 		}
 	}
 
-	// Iterate over each texel and normalize vectors, calculate texel radius
-	for(int32 Y = 0;Y < TextureMapping->CachedSizeY;Y++)
-	{
-		for(int32 X = 0;X < TextureMapping->CachedSizeX;X++)
-		{
-			bool bDebugThisTexel = false;
-#if ALLOW_LIGHTMAP_SAMPLE_DEBUGGING
-			if (bDebugThisMapping
-				&& Y == Scene.DebugInput.LocalY
-				&& X == Scene.DebugInput.LocalX)
-			{
-				bDebugThisTexel = true;
-			}
-#endif
-			FGatheredLightMapSample& CurrentLightSample = LightMapData(X,Y);
-
-			bool bFoundValidCorner = false;
-			const FTexelToCorners& TexelToCorners = TexelToCornersMap(X, Y);
-			for (int32 CornerIndex = 0; CornerIndex < NumTexelCorners; CornerIndex++)
-			{
-				bFoundValidCorner = bFoundValidCorner || TexelToCorners.bValid[CornerIndex];
-			}
-
-			FTexelToVertex& TexelToVertex = TexelToVertexMap(X,Y);
-			if (TexelToVertex.TotalSampleWeight > 0.0f || bFoundValidCorner)
-			{
-				// Use a corner if none of the other samples were valid
-				if (TexelToVertex.TotalSampleWeight < DELTA)
-				{
-					for (int32 CornerIndex = 0; CornerIndex < NumTexelCorners; CornerIndex++)
-					{
-						if (TexelToCorners.bValid[CornerIndex])
-						{
-							TexelToVertex.TotalSampleWeight = 1.0f;
-							TexelToVertex.WorldPosition = TexelToCorners.Corners[CornerIndex].WorldPosition;
-							TexelToVertex.WorldTangentX = TexelToCorners.WorldTangentX;
-							TexelToVertex.WorldTangentY = TexelToCorners.WorldTangentY;
-							TexelToVertex.WorldTangentZ = TexelToCorners.WorldTangentZ;
-							TexelToVertex.TriangleNormal = TexelToCorners.WorldTangentZ;
-							break;
-						}
-					}
-				}
-				else if (GeneralSettings.bUseMaxWeight)
-				{
-					// Weighted average
-					TexelToVertex.WorldTangentX = TexelToVertex.WorldTangentX / TexelToVertex.TotalSampleWeight;
-					TexelToVertex.WorldTangentY = TexelToVertex.WorldTangentY / TexelToVertex.TotalSampleWeight;
-					TexelToVertex.WorldTangentZ = TexelToVertex.WorldTangentZ / TexelToVertex.TotalSampleWeight;
-					TexelToVertex.TriangleNormal = TexelToVertex.TriangleNormal / TexelToVertex.TotalSampleWeight;
-
-					// Weighted average of opposing vectors can result in a zero vector, fixup with corner
-					if (bFoundValidCorner 
-						&& (TexelToVertex.WorldTangentX.SizeSquared3() < KINDA_SMALL_NUMBER
-							|| TexelToVertex.WorldTangentZ.SizeSquared3() < KINDA_SMALL_NUMBER
-							|| TexelToVertex.TriangleNormal.SizeSquared3() < KINDA_SMALL_NUMBER))
-					{
-						for (int32 CornerIndex = 0; CornerIndex < NumTexelCorners; CornerIndex++)
-						{
-							if (TexelToCorners.bValid[CornerIndex])
-							{
-								TexelToVertex.WorldTangentX = TexelToCorners.WorldTangentX;
-								TexelToVertex.WorldTangentY = TexelToCorners.WorldTangentY;
-								TexelToVertex.WorldTangentZ = TexelToCorners.WorldTangentZ;
-								TexelToVertex.TriangleNormal = TexelToCorners.WorldTangentZ;
-								break;
-							}
-						}
-					}
-				}
-
-				// Mark the texel as mapped to some geometry in the scene
-				CurrentLightSample.bIsMapped = true;
-
-				if (MaterialSettings.bUseNormalMapsForLighting && TextureMapping->Mesh->HasImportedNormal(TexelToVertex.ElementIndex))
-				{
-					const FVector4 TangentNormal = TextureMapping->Mesh->EvaluateNormal(TexelToVertex.TextureCoordinates[0], TexelToVertex.ElementIndex);
-
-					const FVector4 WorldTangentRow0(TexelToVertex.WorldTangentX.X, TexelToVertex.WorldTangentY.X, TexelToVertex.WorldTangentZ.X);
-					const FVector4 WorldTangentRow1(TexelToVertex.WorldTangentX.Y, TexelToVertex.WorldTangentY.Y, TexelToVertex.WorldTangentZ.Y);
-					const FVector4 WorldTangentRow2(TexelToVertex.WorldTangentX.Z, TexelToVertex.WorldTangentY.Z, TexelToVertex.WorldTangentZ.Z);
-					const FVector4 WorldVector(
-						Dot3(WorldTangentRow0, TangentNormal),
-						Dot3(WorldTangentRow1, TangentNormal),
-						Dot3(WorldTangentRow2, TangentNormal)
-						);
-
-					TexelToVertex.WorldTangentZ = WorldVector;
-				}
-
-				// Normalize the tangent basis and ensure it is orthonormal
-				TexelToVertex.WorldTangentZ = TexelToVertex.WorldTangentZ.GetUnsafeNormal3();
-
-				const bool bUseVertexNormalForHemisphereGather = TextureMapping->Mesh->UseVertexNormalForHemisphereGather(TexelToVertex.ElementIndex);
-				TexelToVertex.TriangleNormal = bUseVertexNormalForHemisphereGather ? TexelToVertex.WorldTangentZ : TexelToVertex.TriangleNormal.GetUnsafeNormal3();
-				checkSlow(!TexelToVertex.TriangleNormal.ContainsNaN());
-
-				const FVector4 OriginalTangentX = TexelToVertex.WorldTangentX;
-				const FVector4 OriginalTangentY = TexelToVertex.WorldTangentY;
-				
-				TexelToVertex.WorldTangentY = (TexelToVertex.WorldTangentZ ^ TexelToVertex.WorldTangentX).GetUnsafeNormal3();
-				// Maintain handedness
-				if (Dot3(TexelToVertex.WorldTangentY, OriginalTangentY) < 0)
-				{
-					TexelToVertex.WorldTangentY *= -1.0f;
-				}
-				TexelToVertex.WorldTangentX = TexelToVertex.WorldTangentY ^ TexelToVertex.WorldTangentZ;
-				if (Dot3(TexelToVertex.WorldTangentX, OriginalTangentX) < 0)
-				{
-					TexelToVertex.WorldTangentX *= -1.0f;
-				}
-				checkSlow(TexelToVertex.WorldTangentX.IsUnit3());
-				checkSlow(TexelToVertex.WorldTangentY.IsUnit3());
-				checkSlow(TexelToVertex.WorldTangentZ.IsUnit3());
-				checkSlow(TexelToVertex.TriangleNormal.IsUnit3());
-				checkSlow(Dot3(TexelToVertex.WorldTangentZ, TexelToVertex.WorldTangentY) < KINDA_SMALL_NUMBER);
-				checkSlow(Dot3(TexelToVertex.WorldTangentX, TexelToVertex.WorldTangentY) < KINDA_SMALL_NUMBER);
-				checkSlow(Dot3(TexelToVertex.WorldTangentX, TexelToVertex.WorldTangentZ) < KINDA_SMALL_NUMBER);
-
-				// Calculate the bounding radius of the texel
-				// Use the closest corner as it's likely that's on the same section of a split texel 
-				// (A texel shared by multiple UV charts that has sub samples on triangles in different smoothing groups)
-				float MinDistanceSquared = FLT_MAX;
-				if (bFoundValidCorner)
-				{
-					for (int32 CornerIndex = 0; CornerIndex < NumTexelCorners; CornerIndex++)
-					{
-						if (TexelToCorners.bValid[CornerIndex])
-						{
-							const float CornerDistSquared = (TexelToCorners.Corners[CornerIndex].WorldPosition - TexelToVertex.WorldPosition).SizeSquared3();
-							if (CornerDistSquared < MinDistanceSquared)
-							{
-								MinDistanceSquared = CornerDistSquared;
-							}
-						}
-					}
-				}
-				else
-				{
-					MinDistanceSquared = SceneConstants.SmallestTexelRadius;
-				}
-				TexelToVertex.TexelRadius = FMath::Max(FMath::Sqrt(MinDistanceSquared), SceneConstants.SmallestTexelRadius);
-				MappingContext.Stats.NumMappedTexels++;
-
-				{
-					const FFullStaticLightingVertex FullVertex = TexelToVertex.GetFullVertex();
-					const FVector4 TexelCenterOffset = FullVertex.WorldPosition + FullVertex.TriangleNormal * TexelToVertex.TexelRadius * SceneConstants.VisibilityNormalOffsetSampleRadiusScale;
-					
-					FLightRayIntersection Intersections[4];
-					bool bHitBackfaces[4];
-
-					FVector2D CornerSigns[4];
-					CornerSigns[0] = FVector2D(1, 1);
-					CornerSigns[1] = FVector2D(-1, 1);
-					CornerSigns[2] = FVector2D(1, -1);
-					CornerSigns[3] = FVector2D(-1, -1);
-
-					for (int32 CornerIndex = 0; CornerIndex < ARRAY_COUNT(CornerSigns); CornerIndex++)
-					{
-						TraceToTexelCorner(
-							TexelCenterOffset, 
-							FullVertex, 
-							CornerSigns[CornerIndex],
-							// Note: Searching the entire influence of the texel after interpolation, which is 2x the sample radius
-							TexelToVertex.TexelRadius * 2.0f,
-							MappingContext, 
-							Intersections[CornerIndex],
-							bHitBackfaces[CornerIndex],
-							bDebugThisTexel);
-					}
-
-					int32 ClosestIntersectionIndex = INDEX_NONE;
-					float ClosestIntersectionDistanceSq = FLT_MAX;
-
-					int32 ClosestBackfacingIntersectionIndex = INDEX_NONE;
-					// Limit the distance that we will search for an intersecting backface in order to move the shading position to the texel radius
-					float ClosestBackfacingIntersectionDistanceSq = TexelToVertex.TexelRadius * TexelToVertex.TexelRadius;
-
-					for (int32 CornerIndex = 0; CornerIndex < ARRAY_COUNT(CornerSigns); CornerIndex++)
-					{
-						if (Intersections[CornerIndex].bIntersects)
-						{
-							const float DistanceSquared = (Intersections[CornerIndex].IntersectionVertex.WorldPosition - TexelCenterOffset).SizeSquared3();
-
-							if (ClosestIntersectionIndex == INDEX_NONE || DistanceSquared < ClosestIntersectionDistanceSq)
-							{
-								ClosestIntersectionDistanceSq = DistanceSquared;
-								ClosestIntersectionIndex = CornerIndex;
-							}
-
-							if (bHitBackfaces[CornerIndex] && DistanceSquared < ClosestBackfacingIntersectionDistanceSq)
-							{
-								ClosestBackfacingIntersectionDistanceSq = DistanceSquared;
-								ClosestBackfacingIntersectionIndex = CornerIndex;
-							}
-						}
-					}
-					
-					if (ClosestIntersectionIndex != INDEX_NONE)
-					{
-						checkSlow(Intersections[ClosestIntersectionIndex].bIntersects);
-
-						// Mark the texel as intersecting another surface so we can avoid filtering across it later
-						TexelToVertex.bIntersectingSurface = true;
-					}
-
-					// Give preference to moving the shading position outside of backfaces
-					int32 IntersectionIndexForShadingPositionMovement = ClosestBackfacingIntersectionIndex;
-
-					// Note: this is disabled as it causes problems in cracks, the lighting position will be moved inside the object
-					/*
-					// Even if we didn't hit any backfaces, still move the shading position away from an intersecting frontface if it is close enough
-					if (IntersectionIndexForShadingPositionMovement == INDEX_NONE 
-						&& ClosestIntersectionDistanceSq < (TexelToVertex.TexelRadius / 2) * (TexelToVertex.TexelRadius / 2))
-					{
-						IntersectionIndexForShadingPositionMovement = ClosestIntersectionIndex;
-					}*/
-
-					if (IntersectionIndexForShadingPositionMovement != INDEX_NONE)
-					{
-						// Move the shading position outside the surface that is intersecting this texel
-						const FVector4 OffsetShadingPosition = Intersections[IntersectionIndexForShadingPositionMovement].IntersectionVertex.WorldPosition
-							// Move along the intersecting surface's normal but also away from the texel a bit to prevent incorrect self occlusion
-							+ (Intersections[IntersectionIndexForShadingPositionMovement].IntersectionVertex.WorldTangentZ + TexelToVertex.TriangleNormal) * .5f * TexelToVertex.TexelRadius * SceneConstants.VisibilityNormalOffsetSampleRadiusScale;
-
-						// Project back onto plane of texel to avoid incorrect self occlusion
-						TexelToVertex.WorldPosition = OffsetShadingPosition + TexelToVertex.TriangleNormal * Dot3(TexelToVertex.TriangleNormal, TexelToVertex.WorldPosition - OffsetShadingPosition);
-					}
-				}
-			}
-			else
-			{
-				// Mark unmapped texels with the supplied 'UnmappedTexelColor'.
-				CurrentLightSample.AddWeighted(FGatheredLightSampleUtil::AmbientLight<2>(Scene.GeneralSettings.UnmappedTexelColor), 1.0f);
-			}
-		}
-	}
+	AdjustRepresentativeSurfelForTexelsTextureMapping(TextureMapping, TexelToVertexMap, TexelToCornersMap, &LightMapData, MappingContext, bDebugThisMapping);
 }
 
 /** Calculates direct lighting as if all lights were non-area lights, then filters the results in texture space to create approximate soft shadows. */
