Implement bUseFastVoxelization (default to false) in Lightmass to speed up volumetric lightmap voxelization in large scenes. IntersectBox() is implemented in LMkDOP. A kDOP tree is built per static mesh asset (instead of per instance). During voxelization ShouldRefineVoxel() recusively tests voxel boxes against the kDOP tree instances.

diff --git a/Engine/Config/BaseLightmass.ini b/Engine/Config/BaseLightmass.ini
--- a/Engine/Config/BaseLightmass.ini
+++ b/Engine/Config/BaseLightmass.ini
@@ -29,6 +29,8 @@ bRebuildDirtyGeometryForLighting=True
 bUseEmbree=true
 bVerifyEmbree=false
 bUseEmbreePacketTracing=false
+; Whether to use kDOP trees & static mesh instancing to accelerate volumetric lightmap voxelization. Useful in scenes like large forest
+bUseFastVoxelization=false
 bUseFilteredCubemapForSkylight=true
 MappingSurfaceCacheDownsampleFactor=2
 
diff --git a/Engine/Source/Editor/UnrealEd/Private/Lightmass/Lightmass.cpp b/Engine/Source/Editor/UnrealEd/Private/Lightmass/Lightmass.cpp
--- a/Engine/Source/Editor/UnrealEd/Private/Lightmass/Lightmass.cpp
+++ b/Engine/Source/Editor/UnrealEd/Private/Lightmass/Lightmass.cpp
@@ -2180,6 +2180,7 @@ void FLightmassExporter::WriteSceneSettings( Lightmass::FSceneFileHeader& Scene
 		VERIFYLIGHTMASSINI(GConfig->GetBool(TEXT("DevOptions.StaticLighting"), TEXT("bVerifyEmbree"), bConfigBool, GLightmassIni));
 		Scene.GeneralSettings.bVerifyEmbree = Scene.GeneralSettings.bUseEmbree && bConfigBool;
 		VERIFYLIGHTMASSINI(GConfig->GetBool(TEXT("DevOptions.StaticLighting"), TEXT("bUseEmbreePacketTracing"), Scene.GeneralSettings.bUseEmbreePacketTracing, GLightmassIni));
+		VERIFYLIGHTMASSINI(GConfig->GetBool(TEXT("DevOptions.StaticLighting"), TEXT("bUseFastVoxelization"), Scene.GeneralSettings.bUseFastVoxelization, GLightmassIni));
 		VERIFYLIGHTMASSINI(GConfig->GetInt(TEXT("DevOptions.StaticLighting"), TEXT("MappingSurfaceCacheDownsampleFactor"), Scene.GeneralSettings.MappingSurfaceCacheDownsampleFactor, GLightmassIni));
 
 		int32 CheckQualityLevel;
diff --git a/Engine/Source/Programs/UnrealLightmass/Private/CPUSolver/CPUSolver.cpp b/Engine/Source/Programs/UnrealLightmass/Private/CPUSolver/CPUSolver.cpp
--- a/Engine/Source/Programs/UnrealLightmass/Private/CPUSolver/CPUSolver.cpp
+++ b/Engine/Source/Programs/UnrealLightmass/Private/CPUSolver/CPUSolver.cpp
@@ -212,19 +212,9 @@ void BuildStaticLighting(const FGuid& SceneGuid, int32 NumThreads, bool bDumpTex
 
 	UE_LOG(LogLightmass, Log, TEXT("Lighting complete [Startup = %s, Lighting = %s]"), *FPlatformTime::PrettyTime(LightTimeStart - SetupTimeStart), *FPlatformTime::PrettyTime(EndTime - LightTimeStart));
 
-	if ( GReportDetailedStats )
+	if (GReportDetailedStats)
 	{
-		extern volatile uint64 GKDOPParentNodesTraversed;
-		extern volatile uint64 GKDOPLeafNodesTraversed;
-		extern volatile uint64 GKDOPTrianglesTraversed;
-		extern volatile uint64 GKDOPTrianglesTraversedReal;
-		double KDOPTrianglesTraversedRealPercent = (1.0 - ((GKDOPTrianglesTraversed - GKDOPTrianglesTraversedReal) / (GKDOPTrianglesTraversed * 100.0)));
-		UE_LOG(LogLightmass, Log, TEXT("kDOP traversals (in millions): %.3g parents, %.3g leaves, %.3g triangles (%.3g, %.3g%%, real triangles)."),
-			GKDOPParentNodesTraversed / 1000000.0,
-			GKDOPLeafNodesTraversed / 1000000.0,
-			GKDOPTrianglesTraversed / 1000000.0,
-			GKDOPTrianglesTraversedReal / 1000000.0,
-			KDOPTrianglesTraversedRealPercent );
+		LightingSystem.GetAggregateMesh().DumpPostBuildStats();
 	}
 
 	// Transfer back the log to the instigator.
diff --git a/Engine/Source/Programs/UnrealLightmass/Private/ImportExport/Mesh.h b/Engine/Source/Programs/UnrealLightmass/Private/ImportExport/Mesh.h
--- a/Engine/Source/Programs/UnrealLightmass/Private/ImportExport/Mesh.h
+++ b/Engine/Source/Programs/UnrealLightmass/Private/ImportExport/Mesh.h
@@ -44,6 +44,8 @@ public:
 		return LODs[Index];
 	}
 
+	class FDefaultAggregateMesh* VoxelizationMesh = nullptr;
+
 protected:
 	/** array of LODs (same number as FStaticMeshData.NumLODs) */
 	TArray<FStaticMeshLOD>		LODs;
diff --git a/Engine/Source/Programs/UnrealLightmass/Private/Lighting/AdaptiveVolumetricLightmap.cpp b/Engine/Source/Programs/UnrealLightmass/Private/Lighting/AdaptiveVolumetricLightmap.cpp
--- a/Engine/Source/Programs/UnrealLightmass/Private/Lighting/AdaptiveVolumetricLightmap.cpp
+++ b/Engine/Source/Programs/UnrealLightmass/Private/Lighting/AdaptiveVolumetricLightmap.cpp
@@ -127,77 +127,112 @@ bool FStaticLightingSystem::DoesVoxelIntersectSceneGeometry(const FBox& CellBoun
 	const FBox ExpandedCellBoundsSurfaceGeometry = CellBounds.ExpandBy(CellBounds.GetSize() * VolumetricLightmapSettings.VoxelizationCellExpansionForSurfaceGeometry);
 	const FBox ExpandedCellBoundsVolumeGeometry = CellBounds.ExpandBy(CellBounds.GetSize() * VolumetricLightmapSettings.VoxelizationCellExpansionForVolumeGeometry);
 
-	for (int32 MappingIndex = 0; MappingIndex < AllMappings.Num(); MappingIndex++)
+	if (Scene.GeneralSettings.bUseFastVoxelization)
 	{
-		const FStaticLightingMapping* CurrentMapping = AllMappings[MappingIndex];
-		const FStaticLightingTextureMapping* TextureMapping = CurrentMapping->GetTextureMapping();
-		const FStaticLightingMesh* CurrentMesh = CurrentMapping->Mesh;
-		const FBox& ExpandedCellBounds = CurrentMapping->GetVolumeMapping() ? ExpandedCellBoundsVolumeGeometry : ExpandedCellBoundsSurfaceGeometry;
-
-		const bool bMeshBelongsToLOD0 = CurrentMesh->DoesMeshBelongToLOD0();
+		if (VoxelizationSurfaceAggregateMesh->IntersectBox(ExpandedCellBoundsSurfaceGeometry) || VoxelizationVolumeAggregateMesh->IntersectBox(ExpandedCellBoundsVolumeGeometry))
+		{
+			return true;
+		}
 
-		if ((CurrentMesh->LightingFlags & GI_INSTANCE_CASTSHADOW)
-			&& CurrentMesh->BoundingBox.Intersect(ExpandedCellBounds)
-			&& bMeshBelongsToLOD0)
+		for (int32 MeshIndex = 0; MeshIndex < Scene.StaticMeshInstances.Num(); MeshIndex++)
 		{
-			for (int32 TriangleIndex = 0; TriangleIndex < CurrentMesh->NumTriangles; TriangleIndex++)
-			{
-				FStaticLightingVertex Vertices[3];
-				int32 ElementIndex;
-				CurrentMesh->GetTriangle(TriangleIndex, Vertices[0], Vertices[1], Vertices[2], ElementIndex);
+			const FStaticMeshStaticLightingMesh* MeshInstance = &Scene.StaticMeshInstances[MeshIndex];
 
-				if (CurrentMesh->IsElementCastingShadow(ElementIndex) 
-					// Lightmass doesn't handle bounced light from translucency
-					&& !CurrentMesh->IsTranslucent(ElementIndex))
+			if (MeshInstance->StaticMesh->VoxelizationMesh != nullptr)
+			{
+				if (MeshInstance->Mapping->GetVolumeMapping() == nullptr)
+				{
+					if (MeshInstance->StaticMesh->VoxelizationMesh->IntersectBox(ExpandedCellBoundsSurfaceGeometry.TransformBy(MeshInstance->LocalToWorldInverseTranspose.GetTransposed())))
+					{
+						return true;
+					}
+				}
+				else
 				{
-					FTriangle Triangle;
-					Triangle.Vertices[0] = Vertices[0].WorldPosition;
-					Triangle.Vertices[1] = Vertices[1].WorldPosition;
-					Triangle.Vertices[2] = Vertices[2].WorldPosition;
+					if (MeshInstance->StaticMesh->VoxelizationMesh->IntersectBox(ExpandedCellBoundsVolumeGeometry.TransformBy(MeshInstance->LocalToWorldInverseTranspose.GetTransposed())))
+					{
+						return true;
+					}
+				}
+			}
+		}
+
+		return false;
+	}
+	else
+	{
+		for (int32 MappingIndex = 0; MappingIndex < AllMappings.Num(); MappingIndex++)
+		{
+			const FStaticLightingMapping* CurrentMapping = AllMappings[MappingIndex];
+			const FStaticLightingTextureMapping* TextureMapping = CurrentMapping->GetTextureMapping();
+			const FStaticLightingMesh* CurrentMesh = CurrentMapping->Mesh;
+			const FBox& ExpandedCellBounds = CurrentMapping->GetVolumeMapping() ? ExpandedCellBoundsVolumeGeometry : ExpandedCellBoundsSurfaceGeometry;
 
-					FBox TriangleAABB(Triangle.Vertices, ARRAY_COUNT(Triangle.Vertices));
+			const bool bMeshBelongsToLOD0 = CurrentMesh->DoesMeshBelongToLOD0();
+
+			if ((CurrentMesh->LightingFlags & GI_INSTANCE_CASTSHADOW)
+				&& CurrentMesh->BoundingBox.Intersect(ExpandedCellBounds)
+				&& bMeshBelongsToLOD0)
+			{
+				for (int32 TriangleIndex = 0; TriangleIndex < CurrentMesh->NumTriangles; TriangleIndex++)
+				{
+					FStaticLightingVertex Vertices[3];
+					int32 ElementIndex;
+					CurrentMesh->GetTriangle(TriangleIndex, Vertices[0], Vertices[1], Vertices[2], ElementIndex);
 
-					if (ExpandedCellBounds.Intersect(TriangleAABB))
+					if (CurrentMesh->IsElementCastingShadow(ElementIndex)
+						// Lightmass doesn't handle bounced light from translucency
+						&& !CurrentMesh->IsTranslucent(ElementIndex))
 					{
-						const FVector4 TriangleNormal = (Vertices[2].WorldPosition - Vertices[0].WorldPosition) ^ (Vertices[1].WorldPosition - Vertices[0].WorldPosition);
-						const float TriangleArea = 0.5f * TriangleNormal.Size3();
+						FTriangle Triangle;
+						Triangle.Vertices[0] = Vertices[0].WorldPosition;
+						Triangle.Vertices[1] = Vertices[1].WorldPosition;
+						Triangle.Vertices[2] = Vertices[2].WorldPosition;
+
+						FBox TriangleAABB(Triangle.Vertices, ARRAY_COUNT(Triangle.Vertices));
 
-						if (TriangleArea > DELTA)
+						if (ExpandedCellBounds.Intersect(TriangleAABB))
 						{
-							if (TextureMapping)
+							const FVector4 TriangleNormal = (Vertices[2].WorldPosition - Vertices[0].WorldPosition) ^ (Vertices[1].WorldPosition - Vertices[0].WorldPosition);
+							const float TriangleArea = 0.5f * TriangleNormal.Size3();
+
+							if (TriangleArea > DELTA)
 							{
-								// Triangle vertices in lightmap UV space, scaled by the lightmap resolution
-								const FVector2D Vertex0 = Vertices[0].TextureCoordinates[TextureMapping->LightmapTextureCoordinateIndex] * FVector2D(TextureMapping->SizeX, TextureMapping->SizeY);
-								const FVector2D Vertex1 = Vertices[1].TextureCoordinates[TextureMapping->LightmapTextureCoordinateIndex] * FVector2D(TextureMapping->SizeX, TextureMapping->SizeY);
-								const FVector2D Vertex2 = Vertices[2].TextureCoordinates[TextureMapping->LightmapTextureCoordinateIndex] * FVector2D(TextureMapping->SizeX, TextureMapping->SizeY);
-
-								// Area in lightmap space, or the number of lightmap texels covered by this triangle
-								const float LightmapTriangleArea = FMath::Abs(
-									Vertex0.X * (Vertex1.Y - Vertex2.Y)
-									+ Vertex1.X * (Vertex2.Y - Vertex0.Y)
-									+ Vertex2.X * (Vertex0.Y - Vertex1.Y));
-
-								const float TexelDensity = LightmapTriangleArea / TriangleArea;
-								// Skip texture lightmapped triangles whose texel density is less than one texel per the area of a right triangle formed by the child voxel.
-								// If surface lighting is being calculated at a low resolution, it's unlikely that the volume near that surface needs to have detailed lighting.
-								if (TexelDensity < SurfaceLightmapDensityThreshold)
+								if (TextureMapping)
 								{
-									continue;
+									// Triangle vertices in lightmap UV space, scaled by the lightmap resolution
+									const FVector2D Vertex0 = Vertices[0].TextureCoordinates[TextureMapping->LightmapTextureCoordinateIndex] * FVector2D(TextureMapping->SizeX, TextureMapping->SizeY);
+									const FVector2D Vertex1 = Vertices[1].TextureCoordinates[TextureMapping->LightmapTextureCoordinateIndex] * FVector2D(TextureMapping->SizeX, TextureMapping->SizeY);
+									const FVector2D Vertex2 = Vertices[2].TextureCoordinates[TextureMapping->LightmapTextureCoordinateIndex] * FVector2D(TextureMapping->SizeX, TextureMapping->SizeY);
+
+									// Area in lightmap space, or the number of lightmap texels covered by this triangle
+									const float LightmapTriangleArea = FMath::Abs(
+										Vertex0.X * (Vertex1.Y - Vertex2.Y)
+										+ Vertex1.X * (Vertex2.Y - Vertex0.Y)
+										+ Vertex2.X * (Vertex0.Y - Vertex1.Y));
+
+									const float TexelDensity = LightmapTriangleArea / TriangleArea;
+									// Skip texture lightmapped triangles whose texel density is less than one texel per the area of a right triangle formed by the child voxel.
+									// If surface lighting is being calculated at a low resolution, it's unlikely that the volume near that surface needs to have detailed lighting.
+									if (TexelDensity < SurfaceLightmapDensityThreshold)
+									{
+										continue;
+									}
 								}
-							}
 
-							if (IntersectTriangleAndAABB(Triangle, ExpandedCellBounds))
-							{
-								return true;
+								if (IntersectTriangleAndAABB(Triangle, ExpandedCellBounds))
+								{
+									return true;
+								}
 							}
 						}
 					}
 				}
 			}
 		}
-	}
 
-	return false;
+		return false;
+	}
 }
 
 bool FStaticLightingSystem::ShouldRefineVoxel(int32 TreeDepth, const FBox& CellBounds, const TArray<FVector>& VoxelTestPositions, bool bDebugThisVoxel) const
@@ -293,64 +328,76 @@ bool FStaticLightingSystem::ShouldRefineVoxel(int32 TreeDepth, const FBox& CellB
 		&& LandscapeMappings.Num() > 0
 		&& VolumetricLightmapSettings.bCullBricksBelowLandscape)
 	{
-		TArray<FVector, TInlineAllocator<100>> TestPositions;
-		TArray<bool, TInlineAllocator<100>> PositionUnderLandscape;
-
-		int32 TestResolution = 10;
-		TestPositions.Empty(TestResolution * TestResolution);
-		PositionUnderLandscape.Empty(TestResolution * TestResolution);
-		PositionUnderLandscape.AddZeroed(TestResolution * TestResolution);
-
-		for (int32 Y = 0; Y < TestResolution; Y++)
+		if (Scene.GeneralSettings.bUseFastVoxelization)
 		{
-			for (int32 X = 0; X < TestResolution; X++)
+			FBox StretchedCellBounds(CellBounds.Min, FVector(CellBounds.Max.X, CellBounds.Max.Y, LandscapeCullingVoxelizationAggregateMesh->GetBounds().Max.Z));
+
+			if (LandscapeCullingVoxelizationAggregateMesh->IntersectBox(StretchedCellBounds) && !LandscapeCullingVoxelizationAggregateMesh->IntersectBox(CellBounds))
 			{
-				const FVector TestPosition = CellBounds.Min + FVector(X / (float)TestResolution, Y / (float)TestResolution, 1.0f) * CellBounds.GetSize();
-				TestPositions.Add(TestPosition);
+				return false;
 			}
 		}
-
-		for (int32 MappingIndex = 0; MappingIndex < LandscapeMappings.Num(); MappingIndex++)
+		else
 		{
-			const FStaticLightingMapping* CurrentMapping = LandscapeMappings[MappingIndex];
-			const FStaticLightingMesh* CurrentMesh = CurrentMapping->Mesh;
+			TArray<FVector, TInlineAllocator<100>> TestPositions;
+			TArray<bool, TInlineAllocator<100>> PositionUnderLandscape;
 
-			if ((CurrentMesh->LightingFlags & GI_INSTANCE_CASTSHADOW)
-				&& CurrentMesh->BoundingBox.IntersectXY(CellBounds))
+			int32 TestResolution = 10;
+			TestPositions.Empty(TestResolution * TestResolution);
+			PositionUnderLandscape.Empty(TestResolution * TestResolution);
+			PositionUnderLandscape.AddZeroed(TestResolution * TestResolution);
+
+			for (int32 Y = 0; Y < TestResolution; Y++)
 			{
-				for (int32 TriangleIndex = 0; TriangleIndex < CurrentMesh->NumTriangles; TriangleIndex++)
+				for (int32 X = 0; X < TestResolution; X++)
 				{
-					FStaticLightingVertex Vertices[3];
-					int32 ElementIndex;
-					CurrentMesh->GetTriangle(TriangleIndex, Vertices[0], Vertices[1], Vertices[2], ElementIndex);
+					const FVector TestPosition = CellBounds.Min + FVector(X / (float)TestResolution, Y / (float)TestResolution, 1.0f) * CellBounds.GetSize();
+					TestPositions.Add(TestPosition);
+				}
+			}
+
+			for (int32 MappingIndex = 0; MappingIndex < LandscapeMappings.Num(); MappingIndex++)
+			{
+				const FStaticLightingMapping* CurrentMapping = LandscapeMappings[MappingIndex];
+				const FStaticLightingMesh* CurrentMesh = CurrentMapping->Mesh;
 
-					if (CurrentMesh->IsElementCastingShadow(ElementIndex))
+				if ((CurrentMesh->LightingFlags & GI_INSTANCE_CASTSHADOW)
+					&& CurrentMesh->BoundingBox.IntersectXY(CellBounds))
+				{
+					for (int32 TriangleIndex = 0; TriangleIndex < CurrentMesh->NumTriangles; TriangleIndex++)
 					{
-						FTriangle Triangle;
-						Triangle.Vertices[0] = Vertices[0].WorldPosition;
-						Triangle.Vertices[1] = Vertices[1].WorldPosition;
-						Triangle.Vertices[2] = Vertices[2].WorldPosition;
+						FStaticLightingVertex Vertices[3];
+						int32 ElementIndex;
+						CurrentMesh->GetTriangle(TriangleIndex, Vertices[0], Vertices[1], Vertices[2], ElementIndex);
 
-						for (int32 PointIndex = 0; PointIndex < PositionUnderLandscape.Num(); PointIndex++)
+						if (CurrentMesh->IsElementCastingShadow(ElementIndex))
 						{
-							if (PointUnderTriangle(TestPositions[PointIndex], Triangle))
+							FTriangle Triangle;
+							Triangle.Vertices[0] = Vertices[0].WorldPosition;
+							Triangle.Vertices[1] = Vertices[1].WorldPosition;
+							Triangle.Vertices[2] = Vertices[2].WorldPosition;
+
+							for (int32 PointIndex = 0; PointIndex < PositionUnderLandscape.Num(); PointIndex++)
 							{
-								PositionUnderLandscape[PointIndex] = true;
+								if (PointUnderTriangle(TestPositions[PointIndex], Triangle))
+								{
+									PositionUnderLandscape[PointIndex] = true;
+								}
 							}
 						}
 					}
 				}
 			}
-		}
 
-		bool bAllPointsUnderLandscape = true;
+			bool bAllPointsUnderLandscape = true;
 
-		for (int32 PointIndex = 0; PointIndex < PositionUnderLandscape.Num(); PointIndex++)
-		{
-			bAllPointsUnderLandscape = bAllPointsUnderLandscape && PositionUnderLandscape[PointIndex];
-		}
+			for (int32 PointIndex = 0; PointIndex < PositionUnderLandscape.Num(); PointIndex++)
+			{
+				bAllPointsUnderLandscape = bAllPointsUnderLandscape && PositionUnderLandscape[PointIndex];
+			}
 
-		return !bAllPointsUnderLandscape;
+			return !bAllPointsUnderLandscape;
+		}
 	}
 
 	return bVoxelIntersectsScene;
diff --git a/Engine/Source/Programs/UnrealLightmass/Private/Lighting/Collision.cpp b/Engine/Source/Programs/UnrealLightmass/Private/Lighting/Collision.cpp
--- a/Engine/Source/Programs/UnrealLightmass/Private/Lighting/Collision.cpp
+++ b/Engine/Source/Programs/UnrealLightmass/Private/Lighting/Collision.cpp
@@ -7,25 +7,6 @@
 namespace Lightmass
 {
 
-// The default number of triangles to store in each leaf
-#define DEFAULT_MAX_TRIS_PER_LEAF 4
-
-/** Number of triangles in the aggregate mesh kDOP. */
-int32 GKDOPTriangles = 0;
-/** Number of internal nodes in the aggregate mesh kDOP. */
-int32 GKDOPNodes = 0;
-/** Number of leaf nodes in the aggregate mesh kDOP. */
-int32 GKDOPNumLeaves = 0;
-/** Maximum number of triangles per leaf node during the splitting process */
-int32 GKDOPMaxTrisPerLeaf = DEFAULT_MAX_TRIS_PER_LEAF;
-/** Total number of kDOP internal nodes traversed when tracing rays. */
-volatile uint64 GKDOPParentNodesTraversed = 0;
-/** Total number of kDOP leaf nodes traversed when tracing rays. */
-volatile uint64 GKDOPLeafNodesTraversed = 0;
-/** Total number of kDOP triangles tested when tracing rays. */
-volatile uint64 GKDOPTrianglesTraversed = 0;
-volatile uint64 GKDOPTrianglesTraversedReal = 0;
-
 FStaticLightingAggregateMesh::FStaticLightingAggregateMesh(const FScene& InScene):
 	Scene(InScene),
 	bHasShadowCastingPrimitives(false),
@@ -145,6 +126,83 @@ void FDefaultAggregateMesh::AddMesh(const FStaticLightingMesh* Mesh, const FStat
 	}
 }
 
+void FDefaultAggregateMesh::AddMeshForVoxelization(const FStaticLightingMesh* Mesh, const FStaticLightingMapping* Mapping)
+{
+	if (Mesh->LightingFlags & GI_INSTANCE_CASTSHADOW && Mesh->DoesMeshBelongToLOD0())
+	{
+		SceneBounds = SceneBounds + Mesh->BoundingBox;
+
+		const FStaticLightingTextureMapping* TextureMapping = Mapping ? Mapping->GetTextureMapping() : NULL;
+		const int32 BaseVertexIndex = Vertices.Num();
+		MeshInfos.Add(new FStaticLightingMeshInfo(BaseVertexIndex, Mesh));
+		Vertices.AddUninitialized(Mesh->NumVertices);
+		UVs.AddZeroed(Mesh->NumVertices);
+		LightmapUVs.AddZeroed(Mesh->NumVertices);
+
+		const uint32 MeshLODIndices = Mesh->GetLODIndices();
+		const uint32 MeshHLODRange = Mesh->GetHLODRange();
+
+		for (int32 TriangleIndex = 0; TriangleIndex < Mesh->NumTriangles; TriangleIndex++)
+		{
+			// Read the triangle from the mesh.
+			FStaticLightingVertex V0;
+			FStaticLightingVertex V1;
+			FStaticLightingVertex V2;
+			int32 ElementIndex;
+			Mesh->GetTriangle(TriangleIndex, V0, V1, V2, ElementIndex);
+
+			if (Mesh->IsElementCastingShadow(ElementIndex) && !Mesh->IsTranslucent(ElementIndex))
+			{
+				int32 I0 = 0;
+				int32 I1 = 0;
+				int32 I2 = 0;
+				Mesh->GetTriangleIndices(TriangleIndex, I0, I1, I2);
+
+				check(I0 <= Mesh->NumVertices && I1 <= Mesh->NumVertices && I2 <= Mesh->NumVertices);
+
+				const bool bTwoSided = Mesh->IsTwoSided(ElementIndex) || Mesh->IsCastingShadowAsTwoSided();
+				const bool bStaticAndOpaque = !Mesh->IsMasked(ElementIndex) && !Mesh->IsTranslucent(ElementIndex) && !Mesh->bMovable;
+
+				Vertices[BaseVertexIndex + I0] = V0.WorldPosition;
+				Vertices[BaseVertexIndex + I1] = V1.WorldPosition;
+				Vertices[BaseVertexIndex + I2] = V2.WorldPosition;
+				UVs[BaseVertexIndex + I0] = V0.TextureCoordinates[Mesh->TextureCoordinateIndex];
+				UVs[BaseVertexIndex + I1] = V1.TextureCoordinates[Mesh->TextureCoordinateIndex];
+				UVs[BaseVertexIndex + I2] = V2.TextureCoordinates[Mesh->TextureCoordinateIndex];
+				if (TextureMapping)
+				{
+					LightmapUVs[BaseVertexIndex + I0] = V0.TextureCoordinates[TextureMapping->LightmapTextureCoordinateIndex];
+					LightmapUVs[BaseVertexIndex + I1] = V1.TextureCoordinates[TextureMapping->LightmapTextureCoordinateIndex];
+					LightmapUVs[BaseVertexIndex + I2] = V2.TextureCoordinates[TextureMapping->LightmapTextureCoordinateIndex];
+				}
+
+				// Compute the triangle's normal.
+				const FVector4 TriangleNormal = (V2.WorldPosition - V0.WorldPosition) ^ (V1.WorldPosition - V0.WorldPosition);
+
+				// Compute the triangle area.
+				const float TriangleArea = TriangleNormal.Size3() * 0.5f;
+
+				// Ignore zero area triangles.
+				if (TriangleArea > TRIANGLE_AREA_THRESHOLD)
+				{
+					const int32 PayloadIndex = TrianglePayloads.Add(
+						FTriangleSOAPayload(MeshInfos.Last(), Mapping, ElementIndex, BaseVertexIndex + I0, BaseVertexIndex + I1, BaseVertexIndex + I2));
+
+					new(kDOPTriangles) FkDOPBuildCollisionTriangle<uint32>(
+						PayloadIndex, // Use the triangle's material index as an index into TrianglePayloads.
+						V0.WorldPosition, V1.WorldPosition, V2.WorldPosition,
+						Mesh->MeshIndex,
+						MeshLODIndices,
+						MeshHLODRange,
+						bTwoSided,
+						bStaticAndOpaque
+						);
+				}
+			}
+		}
+	}
+}
+
 /**
  * Pre-allocates memory ahead of time, before calling AddMesh() a bunch of times.
  *
@@ -167,17 +225,19 @@ void FDefaultAggregateMesh::PrepareForRaytracing()
 {
 	// Build the kDOP for simple meshes.
 	kDopTree.Build(kDOPTriangles);
-
-	// Log information about the aggregate mesh.
-	UE_LOG(LogLightmass, Log, TEXT("Static lighting kDOP: %u nodes, %u leaves, %u triangles, %u vertices"), GKDOPNodes, GKDOPNumLeaves, GKDOPTriangles, Vertices.Num());
-	UE_LOG(LogLightmass, Log, TEXT("Static lighting kDOP: %.3f%% wasted space in leaves"), ((GKDOPTriangles - kDOPTriangles.Num()) / (float)GKDOPTriangles) * 100.0f);
-
 	kDOPTriangles.Empty();
 	TrianglePayloads.Shrink();
 }
 
 void FDefaultAggregateMesh::DumpStats() const
 {
+	UE_LOG(LogLightmass, Log, TEXT("Preallocated %.1fGb for kDOP nodes and triangles"), kDopTree.kDOPPreallocatedMemory / 1024.0f / 1024.0f / 1024.0f);
+	UE_LOG(LogLightmass, Log, TEXT("Building kDOP took %5.2f seconds."), kDopTree.kDOPBuildTime);
+
+	// Log information about the aggregate mesh.
+	UE_LOG(LogLightmass, Log, TEXT("Static lighting kDOP: %u nodes, %u leaves, %u triangles, %u vertices"), kDopTree.GKDOPNodes, kDopTree.GKDOPNumLeaves, kDopTree.GKDOPTriangles, Vertices.Num());
+	UE_LOG(LogLightmass, Log, TEXT("Static lighting kDOP: %.3f%% wasted space in leaves"), ((kDopTree.GKDOPTriangles - kDOPTriangles.Num()) / (float)kDopTree.GKDOPTriangles) * 100.0f);
+
 	const uint64 kDOPTreeBytes = kDopTree.Nodes.GetAllocatedSize() 
 		+ kDopTree.SOATriangles.GetAllocatedSize()
 		+ kDOPTriangles.GetAllocatedSize()
@@ -195,7 +255,18 @@ void FDefaultAggregateMesh::DumpStats() const
 	UE_LOG(LogLightmass, Log, TEXT("Vertices              : %7.1fMb"), Vertices.GetAllocatedSize() / 1048576.0f);
 	UE_LOG(LogLightmass, Log, TEXT("UVs                   : %7.1fMb"), UVs.GetAllocatedSize() / 1048576.0f);
 	UE_LOG(LogLightmass, Log, TEXT("LightmapUVs           : %7.1fMb"), LightmapUVs.GetAllocatedSize() / 1048576.0f);
-	UE_LOG(LogLightmass, Log, TEXT("Static lighting kDOP: %u nodes, %u leaves, %u triangles, %u vertices, %.1f Mb"), GKDOPNodes, GKDOPNumLeaves, GKDOPTriangles, Vertices.Num(), kDOPTreeBytes / 1048576.0f);
+	UE_LOG(LogLightmass, Log, TEXT("Static lighting kDOP: %u nodes, %u leaves, %u triangles, %u vertices, %.1f Mb"), kDopTree.GKDOPNodes, kDopTree.GKDOPNumLeaves, kDopTree.GKDOPTriangles, Vertices.Num(), kDOPTreeBytes / 1048576.0f);
+}
+
+void FDefaultAggregateMesh::DumpPostBuildStats() const
+{
+	double KDOPTrianglesTraversedRealPercent = (1.0 - ((kDopTree.GKDOPTrianglesTraversed - kDopTree.GKDOPTrianglesTraversedReal) / (kDopTree.GKDOPTrianglesTraversed * 100.0)));
+	UE_LOG(LogLightmass, Log, TEXT("kDOP traversals (in millions): %.3g parents, %.3g leaves, %.3g triangles (%.3g, %.3g%%, real triangles)."),
+		kDopTree.GKDOPParentNodesTraversed / 1000000.0,
+		kDopTree.GKDOPLeafNodesTraversed / 1000000.0,
+		kDopTree.GKDOPTrianglesTraversed / 1000000.0,
+		kDopTree.GKDOPTrianglesTraversedReal / 1000000.0,
+		KDOPTrianglesTraversedRealPercent);
 }
 
 class FStaticLightingAggregateMeshDataProvider
@@ -437,5 +508,14 @@ bool FDefaultAggregateMesh::IntersectLightRay(
 	return ClosestIntersection.bIntersects;
 }
 
+bool FDefaultAggregateMesh::IntersectBox(const FBox Box)  const
+{
+	FLightRay UnusedRay;
+	FStaticLightingAggregateMeshDataProvider kDOPDataProvider(this, UnusedRay);
+	TkDOPBoxCollisionCheck<const FStaticLightingAggregateMeshDataProvider, uint32> kDOPCheck(
+		Box,
+		kDOPDataProvider);
+	return kDopTree.BoxCheck(kDOPCheck);
+}
 
 } //namespace Lightmass
diff --git a/Engine/Source/Programs/UnrealLightmass/Private/Lighting/Collision.h b/Engine/Source/Programs/UnrealLightmass/Private/Lighting/Collision.h
--- a/Engine/Source/Programs/UnrealLightmass/Private/Lighting/Collision.h
+++ b/Engine/Source/Programs/UnrealLightmass/Private/Lighting/Collision.h
@@ -152,6 +152,8 @@ public:
 	// Used to dump additional stats about intersect checks that where done
 	virtual void DumpCheckStats() const {}
 
+	virtual void DumpPostBuildStats() const {}
+
 	/**
 	 * Checks a light ray for intersection with the shadow mesh.
 	 * @param LightRay - The line segment to check for intersection.
@@ -228,12 +230,16 @@ public:
 
 	virtual void AddMesh(const FStaticLightingMesh* Mesh, const FStaticLightingMapping* Mapping) override;
 
+	void AddMeshForVoxelization(const FStaticLightingMesh* Mesh, const FStaticLightingMapping* Mapping);
+
 	virtual void ReserveMemory( int32 NumMeshes, int32 NumVertices, int32 NumTriangles ) override;
 
 	virtual void PrepareForRaytracing() override;
 
 	virtual void DumpStats() const override;
 
+	virtual void DumpPostBuildStats() const override;
+
 	virtual bool IntersectLightRay(
 		const FLightRay& LightRay,
 		bool bFindClosestIntersection,
@@ -242,6 +248,8 @@ public:
 		class FCoherentRayCache& CoherentRayCache,
 		FLightRayIntersection& Intersection) const override;
 
+	bool IntersectBox(const FBox Box) const;
+
 private:
 
 	friend class FStaticLightingAggregateMeshDataProvider;
diff --git a/Engine/Source/Programs/UnrealLightmass/Private/Lighting/LightingMesh.h b/Engine/Source/Programs/UnrealLightmass/Private/Lighting/LightingMesh.h
--- a/Engine/Source/Programs/UnrealLightmass/Private/Lighting/LightingMesh.h
+++ b/Engine/Source/Programs/UnrealLightmass/Private/Lighting/LightingMesh.h
@@ -14,6 +14,7 @@ namespace Lightmass
 class FLight;
 class FScene;
 class FMeshAreaLight;
+class FStaticMeshStaticLightingMesh;
 
 /** The vertex data used to build static lighting. */
 struct FStaticLightingVertex: public FStaticLightingVertexData
@@ -332,6 +333,10 @@ protected:
 
 public:
 
+	virtual FStaticMeshStaticLightingMesh* GetInstanceableStaticMesh() { return nullptr; }
+
+	virtual const FStaticMeshStaticLightingMesh* GetInstanceableStaticMesh() const { return nullptr; }
+
 	/** Virtual destructor. */
 	virtual ~FStaticLightingMesh() {}
 
diff --git a/Engine/Source/Programs/UnrealLightmass/Private/Lighting/LightingSystem.cpp b/Engine/Source/Programs/UnrealLightmass/Private/Lighting/LightingSystem.cpp
--- a/Engine/Source/Programs/UnrealLightmass/Private/Lighting/LightingSystem.cpp
+++ b/Engine/Source/Programs/UnrealLightmass/Private/Lighting/LightingSystem.cpp
@@ -153,6 +153,9 @@ FStaticLightingSystem::FStaticLightingSystem(const FLightingBuildOptions& InOpti
 ,	SecondBouncePhotonMap(FVector4(0,0,0), HALF_WORLD_MAX)
 ,	IrradiancePhotonMap(FVector4(0,0,0), HALF_WORLD_MAX)
 ,	AggregateMesh(NULL)
+,	VoxelizationSurfaceAggregateMesh(NULL)
+,	VoxelizationVolumeAggregateMesh(NULL)
+,	LandscapeCullingVoxelizationAggregateMesh(NULL)
 ,	Scene(InScene)
 ,	NumTexelsCompleted(0)
 ,	NumOutstandingVolumeDataLayers(0)
@@ -163,7 +166,7 @@ FStaticLightingSystem::FStaticLightingSystem(const FLightingBuildOptions& InOpti
 ,	Exporter(InExporter)
 {
 	const double SceneSetupStart = FPlatformTime::Seconds();
-	UE_LOG(LogLightmass, Log, TEXT("FStaticLightingSystem started using GKDOPMaxTrisPerLeaf: %d"), GKDOPMaxTrisPerLeaf );
+	UE_LOG(LogLightmass, Log, TEXT("FStaticLightingSystem started using GKDOPMaxTrisPerLeaf: %d"), DEFAULT_MAX_TRIS_PER_LEAF);
 
 	ValidateSettings(InScene);
 
@@ -363,25 +366,110 @@ FStaticLightingSystem::FStaticLightingSystem(const FLightingBuildOptions& InOpti
 	// Add all meshes to the kDOP.
 	AggregateMesh->ReserveMemory(NumMeshes, NumVertices, NumTriangles);
 
+	if (Scene.GeneralSettings.bUseFastVoxelization)
+	{
+		VoxelizationSurfaceAggregateMesh = new FDefaultAggregateMesh(Scene);
+		VoxelizationVolumeAggregateMesh = new FDefaultAggregateMesh(Scene);
+		LandscapeCullingVoxelizationAggregateMesh = new FDefaultAggregateMesh(Scene);
+	}
+
 	for (int32 MappingIndex = 0; MappingIndex < InScene.FluidMappings.Num(); MappingIndex++)
 	{
 		FFluidSurfaceStaticLightingTextureMapping* Mapping = &InScene.FluidMappings[MappingIndex];
 		AggregateMesh->AddMesh(Mapping->Mesh, Mapping);
+
+		if (Scene.GeneralSettings.bUseFastVoxelization)
+		{
+			if (Mapping->GetVolumeMapping() == nullptr)
+			{
+				VoxelizationSurfaceAggregateMesh->AddMeshForVoxelization(Mapping->Mesh, Mapping);
+			}
+			else
+			{
+				VoxelizationVolumeAggregateMesh->AddMeshForVoxelization(Mapping->Mesh, Mapping);
+			}
+		}
 	}
 	for (int32 MappingIndex = 0; MappingIndex < InScene.LandscapeMappings.Num(); MappingIndex++)
 	{
 		FLandscapeStaticLightingTextureMapping* Mapping = &InScene.LandscapeMappings[MappingIndex];
 		AggregateMesh->AddMesh(Mapping->Mesh, Mapping);
+
+		if (Scene.GeneralSettings.bUseFastVoxelization)
+		{
+			if (Mapping->GetVolumeMapping() == nullptr)
+			{
+				VoxelizationSurfaceAggregateMesh->AddMeshForVoxelization(Mapping->Mesh, Mapping);
+			}
+			else
+			{
+				VoxelizationVolumeAggregateMesh->AddMeshForVoxelization(Mapping->Mesh, Mapping);
+			}
+
+			LandscapeCullingVoxelizationAggregateMesh->AddMeshForVoxelization(Mapping->Mesh, Mapping);
+		}
 	}
 	for( int32 MeshIdx=0; MeshIdx < InScene.BspMappings.Num(); ++MeshIdx )
 	{
 		FBSPSurfaceStaticLighting* BSPMapping = &InScene.BspMappings[MeshIdx];
 		AggregateMesh->AddMesh(BSPMapping, &BSPMapping->Mapping);
+
+		if (Scene.GeneralSettings.bUseFastVoxelization)
+		{
+			if (BSPMapping->Mapping.GetVolumeMapping() == nullptr)
+			{
+				VoxelizationSurfaceAggregateMesh->AddMeshForVoxelization(BSPMapping, &BSPMapping->Mapping);
+			}
+			else
+			{
+				VoxelizationVolumeAggregateMesh->AddMeshForVoxelization(BSPMapping, &BSPMapping->Mapping);
+			}
+		}
 	}
 	for (int32 MeshIndex = 0; MeshIndex < InScene.StaticMeshInstances.Num(); MeshIndex++)
 	{
 		FStaticMeshStaticLightingMesh* MeshInstance = &InScene.StaticMeshInstances[MeshIndex];
 		AggregateMesh->AddMesh(MeshInstance, MeshInstance->Mapping);
+
+		if (Scene.GeneralSettings.bUseFastVoxelization)
+		{
+			if (MeshInstance->GetInstanceableStaticMesh() != nullptr)
+			{
+				if (MeshInstance->StaticMesh->VoxelizationMesh == nullptr)
+				{
+					if (MeshInstance->LightingFlags & GI_INSTANCE_CASTSHADOW && MeshInstance->DoesMeshBelongToLOD0())
+					{
+						MeshInstance->StaticMesh->VoxelizationMesh = new FDefaultAggregateMesh(Scene);
+
+						{
+							auto LocalToWorld = MeshInstance->LocalToWorld;
+							auto LocalToWorldInverseTranspose = MeshInstance->LocalToWorldInverseTranspose;
+
+							MeshInstance->LocalToWorld = FMatrix::Identity;
+							MeshInstance->LocalToWorldInverseTranspose = FMatrix::Identity;
+
+							MeshInstance->StaticMesh->VoxelizationMesh->AddMeshForVoxelization(MeshInstance, MeshInstance->Mapping);
+							MeshInstance->StaticMesh->VoxelizationMesh->PrepareForRaytracing();
+
+							MeshInstance->LocalToWorld = LocalToWorld;
+							MeshInstance->LocalToWorldInverseTranspose = LocalToWorldInverseTranspose;
+						}
+					}
+				}
+			}
+			else
+			{
+				// For non-instanceable static meshes (splines), add them to the aggregate scene voxelization mesh
+				if (MeshInstance->Mapping->GetVolumeMapping() == nullptr)
+				{
+					VoxelizationSurfaceAggregateMesh->AddMeshForVoxelization(MeshInstance, MeshInstance->Mapping);
+				}
+				else
+				{
+					VoxelizationVolumeAggregateMesh->AddMeshForVoxelization(MeshInstance, MeshInstance->Mapping);
+				}
+			}
+		}
 	}
 
 	// Comparing mappings based on cost, descending.
@@ -477,6 +565,14 @@ FStaticLightingSystem::FStaticLightingSystem(const FLightingBuildOptions& InOpti
 	AggregateMesh->PrepareForRaytracing();
 	AggregateMesh->DumpStats();
 
+	if (Scene.GeneralSettings.bUseFastVoxelization)
+	{
+		VoxelizationSurfaceAggregateMesh->PrepareForRaytracing();
+		VoxelizationVolumeAggregateMesh->PrepareForRaytracing();
+
+		LandscapeCullingVoxelizationAggregateMesh->PrepareForRaytracing();
+	}
+
 	NumCompletedRadiosityIterationMappings.Empty(GeneralSettings.NumSkyLightingBounces);
 	NumCompletedRadiosityIterationMappings.AddDefaulted(GeneralSettings.NumSkyLightingBounces);
 
@@ -491,6 +587,28 @@ FStaticLightingSystem::~FStaticLightingSystem()
 {
 	delete AggregateMesh;
 	AggregateMesh = NULL;
+
+	if (Scene.GeneralSettings.bUseFastVoxelization)
+	{
+		delete VoxelizationSurfaceAggregateMesh;
+		VoxelizationSurfaceAggregateMesh = NULL;
+
+		delete VoxelizationVolumeAggregateMesh;
+		VoxelizationVolumeAggregateMesh = NULL;
+
+		for (int32 MeshIndex = 0; MeshIndex < Scene.StaticMeshInstances.Num(); MeshIndex++)
+		{
+			const FStaticMeshStaticLightingMesh* MeshInstance = &Scene.StaticMeshInstances[MeshIndex];
+			if (MeshInstance->StaticMesh->VoxelizationMesh != nullptr)
+			{
+				delete MeshInstance->StaticMesh->VoxelizationMesh;
+				MeshInstance->StaticMesh->VoxelizationMesh = nullptr;
+			}
+		}
+
+		delete LandscapeCullingVoxelizationAggregateMesh;
+		LandscapeCullingVoxelizationAggregateMesh = NULL;
+	}
 }
 
 /**
diff --git a/Engine/Source/Programs/UnrealLightmass/Private/Lighting/LightingSystem.h b/Engine/Source/Programs/UnrealLightmass/Private/Lighting/LightingSystem.h
--- a/Engine/Source/Programs/UnrealLightmass/Private/Lighting/LightingSystem.h
+++ b/Engine/Source/Programs/UnrealLightmass/Private/Lighting/LightingSystem.h
@@ -2690,7 +2690,11 @@ private:
 
 	/** The aggregate mesh used for raytracing. */
 	FStaticLightingAggregateMeshType* AggregateMesh;
-	
+
+	FDefaultAggregateMesh* VoxelizationSurfaceAggregateMesh;
+	FDefaultAggregateMesh* VoxelizationVolumeAggregateMesh;
+	FDefaultAggregateMesh* LandscapeCullingVoxelizationAggregateMesh;
+
 	/** The input scene describing geometry, materials and lights. */
 	const FScene& Scene; 
 
diff --git a/Engine/Source/Programs/UnrealLightmass/Private/Lighting/StaticMesh.h b/Engine/Source/Programs/UnrealLightmass/Private/Lighting/StaticMesh.h
--- a/Engine/Source/Programs/UnrealLightmass/Private/Lighting/StaticMesh.h
+++ b/Engine/Source/Programs/UnrealLightmass/Private/Lighting/StaticMesh.h
@@ -19,6 +19,10 @@ namespace Lightmass
 
 		FStaticLightingMapping* Mapping;
 
+		virtual FStaticMeshStaticLightingMesh* GetInstanceableStaticMesh() override { return !bIsSplineMesh ? this : nullptr; }
+
+		virtual const FStaticMeshStaticLightingMesh* GetInstanceableStaticMesh() const override { return !bIsSplineMesh ? this : nullptr; }
+
 		// FStaticLightingMesh interface.
 		/**
 		 *	Returns the Guid for the object associated with this lighting mesh.
@@ -63,8 +67,6 @@ namespace Lightmass
 
 		virtual void Import( class FLightmassImporter& Importer );
 
-	private:
-
 		/** The inverse transpose of the primitive's local to world transform. */
 		FMatrix LocalToWorldInverseTranspose;
 	};
diff --git a/Engine/Source/Programs/UnrealLightmass/Private/LightmassCore/Math/LMkDOP.h b/Engine/Source/Programs/UnrealLightmass/Private/LightmassCore/Math/LMkDOP.h
--- a/Engine/Source/Programs/UnrealLightmass/Private/LightmassCore/Math/LMkDOP.h
+++ b/Engine/Source/Programs/UnrealLightmass/Private/LightmassCore/Math/LMkDOP.h
@@ -18,22 +18,8 @@ namespace Lightmass
 #define SLOW_KDOP_STATS(...) 
 // Whether to cache Centroid, LocalNormal and such for build triangles. Peak memory vs. minor build time trade-off.
 #define CACHE_BUILD_TEMPORARIES 0
-
-/** Number of triangles in the aggregate mesh kDOP. */
-extern int32 GKDOPTriangles;
-/** Number of internal nodes in the aggregate mesh kDOP. */
-extern int32 GKDOPNodes;
-/** Number of leaf nodes in the aggregate mesh kDOP. */
-extern int32 GKDOPNumLeaves;
-/** Maximum number of triangles per leaf node during the splitting process */
-extern int32 GKDOPMaxTrisPerLeaf;
-/** Total number of kDOP internal nodes traversed when tracing rays. */
-extern volatile uint64 GKDOPParentNodesTraversed;
-/** Total number of kDOP leaf nodes traversed when tracing rays. */
-extern volatile uint64 GKDOPLeafNodesTraversed;
-/** Total number of kDOP triangles tested when tracing rays. */
-extern volatile uint64 GKDOPTrianglesTraversed;
-extern volatile uint64 GKDOPTrianglesTraversedReal;
+// The default number of triangles to store in each leaf
+#define DEFAULT_MAX_TRIS_PER_LEAF 4
 
 struct FHitResult
 {
@@ -501,6 +487,7 @@ public:
 template <typename COLL_DATA_PROVIDER,typename KDOP_IDX_TYPE> struct TkDOPNode;
 template <typename COLL_DATA_PROVIDER,typename KDOP_IDX_TYPE> struct TkDOPTree;
 template <typename COLL_DATA_PROVIDER,typename KDOP_IDX_TYPE> struct TkDOPLineCollisionCheck;
+template <typename COLL_DATA_PROVIDER,typename KDOP_IDX_TYPE> struct TkDOPBoxCollisionCheck;
 
 /**
  * Holds the min/max planes that make up a set of 4 bounding volumes.
@@ -555,6 +542,112 @@ struct FMultiBox
 	#define kDOPArray TChunkedArray
 #endif
 
+namespace BoxTriangleIntersectionInternal
+{
+
+struct FTriangle
+{
+	FVector Vertices[3];
+};
+
+struct FOverlapInterval
+{
+	float Min;
+	float Max;
+};
+
+FORCEINLINE FOverlapInterval GetInterval(const FTriangle& Triangle, const FVector& Vector)
+{
+	FOverlapInterval Result;
+	Result.Min = Result.Max = FVector::DotProduct(Vector, Triangle.Vertices[0]);
+
+	for (int32 i = 1; i < 3; ++i)
+	{
+		float Projection = FVector::DotProduct(Vector, Triangle.Vertices[i]);
+		Result.Min = FMath::Min(Result.Min, Projection);
+		Result.Max = FMath::Max(Result.Max, Projection);
+	}
+
+	return Result;
+}
+
+FORCEINLINE FOverlapInterval GetInterval(const FBox& Box, const FVector& Vector)
+{
+	FVector BoxVertices[8] =
+	{
+		FVector(Box.Min.X, Box.Max.Y, Box.Max.Z),
+		FVector(Box.Min.X, Box.Max.Y, Box.Min.Z),
+		FVector(Box.Min.X, Box.Min.Y, Box.Max.Z),
+		FVector(Box.Min.X, Box.Min.Y, Box.Min.Z),
+		FVector(Box.Max.X, Box.Max.Y, Box.Max.Z),
+		FVector(Box.Max.X, Box.Max.Y, Box.Min.Z),
+		FVector(Box.Max.X, Box.Min.Y, Box.Max.Z),
+		FVector(Box.Max.X, Box.Min.Y, Box.Min.Z)
+	};
+
+	FOverlapInterval Result;
+	Result.Min = Result.Max = FVector::DotProduct(Vector, BoxVertices[0]);
+
+	for (int32 i = 1; i < ARRAY_COUNT(BoxVertices); ++i)
+	{
+		float Projection = FVector::DotProduct(Vector, BoxVertices[i]);
+		Result.Min = FMath::Min(Result.Min, Projection);
+		Result.Max = FMath::Max(Result.Max, Projection);
+	}
+
+	return Result;
+}
+
+FORCEINLINE bool OverlapOnAxis(const FBox& Box, const FTriangle& Triangle, const FVector& Vector)
+{
+	FOverlapInterval A = GetInterval(Box, Vector);
+	FOverlapInterval B = GetInterval(Triangle, Vector);
+	return ((B.Min <= A.Max) && (A.Min <= B.Max));
+}
+
+FORCEINLINE bool IntersectTriangleAndAABB(const FTriangle& Triangle, const FBox& Box)
+{
+	FVector TriangleEdge0 = Triangle.Vertices[1] - Triangle.Vertices[0];
+	FVector TriangleEdge1 = Triangle.Vertices[2] - Triangle.Vertices[1];
+	FVector TriangleEdge2 = Triangle.Vertices[0] - Triangle.Vertices[2];
+
+	FVector BoxNormal0(1.0f, 0.0f, 0.0f);
+	FVector BoxNormal1(0.0f, 1.0f, 0.0f);
+	FVector BoxNormal2(0.0f, 0.0f, 1.0f);
+
+	FVector TestDirections[13] =
+	{
+		// Separating axes from the box normals
+		BoxNormal0,
+		BoxNormal1,
+		BoxNormal2,
+		// One separating axis for the triangle normal
+		FVector::CrossProduct(TriangleEdge0, TriangleEdge1),
+		// Separating axes for the triangle edges
+		FVector::CrossProduct(BoxNormal0, TriangleEdge0),
+		FVector::CrossProduct(BoxNormal0, TriangleEdge1),
+		FVector::CrossProduct(BoxNormal0, TriangleEdge2),
+		FVector::CrossProduct(BoxNormal1, TriangleEdge0),
+		FVector::CrossProduct(BoxNormal1, TriangleEdge1),
+		FVector::CrossProduct(BoxNormal1, TriangleEdge2),
+		FVector::CrossProduct(BoxNormal2, TriangleEdge0),
+		FVector::CrossProduct(BoxNormal2, TriangleEdge1),
+		FVector::CrossProduct(BoxNormal2, TriangleEdge2)
+	};
+
+	for (int i = 0; i < ARRAY_COUNT(TestDirections); ++i)
+	{
+		if (!OverlapOnAxis(Box, Triangle, TestDirections[i]))
+		{
+			// If we don't overlap on a single axis, the shapes do not intersect
+			return false;
+		}
+	}
+
+	return true;
+}
+}
+
 /**
  * A node in the kDOP tree. The node contains the kDOP volume that encompasses
  * it's children and/or triangles
@@ -891,6 +984,77 @@ struct TkDOPNode
 		}
 		return bHit;
 	}
+
+	bool BoxCheckTriangles(TkDOPBoxCollisionCheck<COLL_DATA_PROVIDER, KDOP_IDX_TYPE>& Check) const
+	{
+		for (KDOP_IDX_TYPE SOAIndex = t.StartIndex; SOAIndex < (t.StartIndex + t.NumTriangles); SOAIndex++)
+		{
+			const FTriangleSOA& TriangleSOA = Check.SOATriangles[SOAIndex];
+			for (int32 SubIndex = 0; SubIndex < 4; SubIndex++)
+			{
+				BoxTriangleIntersectionInternal::FTriangle Triangle;
+
+				Triangle.Vertices[0].X = LmVectorGetComponent(TriangleSOA.Positions[0].X, SubIndex);
+				Triangle.Vertices[0].Y = LmVectorGetComponent(TriangleSOA.Positions[0].Y, SubIndex);
+				Triangle.Vertices[0].Z = LmVectorGetComponent(TriangleSOA.Positions[0].Z, SubIndex);
+
+				Triangle.Vertices[1].X = LmVectorGetComponent(TriangleSOA.Positions[1].X, SubIndex);
+				Triangle.Vertices[1].Y = LmVectorGetComponent(TriangleSOA.Positions[1].Y, SubIndex);
+				Triangle.Vertices[1].Z = LmVectorGetComponent(TriangleSOA.Positions[1].Z, SubIndex);
+
+				Triangle.Vertices[2].X = LmVectorGetComponent(TriangleSOA.Positions[2].X, SubIndex);
+				Triangle.Vertices[2].Y = LmVectorGetComponent(TriangleSOA.Positions[2].Y, SubIndex);
+				Triangle.Vertices[2].Z = LmVectorGetComponent(TriangleSOA.Positions[2].Z, SubIndex);
+
+				if (BoxTriangleIntersectionInternal::IntersectTriangleAndAABB(Triangle, Check.Box))
+					return true;
+			}
+		}
+
+		return false;
+	}
+
+	void BoxCheckBounds(TkDOPBoxCollisionCheck<COLL_DATA_PROVIDER, KDOP_IDX_TYPE>& Check, bool NodeHit[2]) const
+	{
+		for (int32 BoxIndex = 0; BoxIndex < 2; BoxIndex++)
+		{
+			// 0: Create constants
+			FVector BoxMin(BoundingVolumes.Min[0][BoxIndex], BoundingVolumes.Min[1][BoxIndex], BoundingVolumes.Min[2][BoxIndex]);
+			FVector BoxMax(BoundingVolumes.Max[0][BoxIndex], BoundingVolumes.Max[1][BoxIndex], BoundingVolumes.Max[2][BoxIndex]);
+
+			FBox ChildBox(BoxMin, BoxMax);
+
+			NodeHit[BoxIndex] = Check.Box.Intersect(ChildBox);
+		}
+	}
+
+	bool BoxCheck(TkDOPBoxCollisionCheck<COLL_DATA_PROVIDER, KDOP_IDX_TYPE>& Check) const
+	{
+		if (bIsLeaf == 0)
+		{
+			// Check both left and right node at the same time.
+			bool NodeHit[2];
+			BoxCheckBounds(Check, NodeHit);
+
+			if (NodeHit[0])
+			{
+				if (Check.Nodes[n.LeftNode].BoxCheck(Check))
+					return true;
+			}
+
+			if (NodeHit[1])
+			{
+				if (Check.Nodes[n.RightNode].BoxCheck(Check))
+					return true;
+			}
+
+			return false;
+		}
+		else
+		{
+			return BoxCheckTriangles(Check);
+		}
+	}
 };
 
 /**
@@ -906,6 +1070,25 @@ struct TkDOPTree
 	/** Exposes node type to clients. */
 	typedef TkDOPNode<DataProviderType,KDOP_IDX_TYPE>	NodeType;
 
+	/** Statistics --------------------------------------*/
+	/** Number of triangles in the aggregate mesh kDOP. */
+	int32 GKDOPTriangles = 0;
+	/** Number of internal nodes in the aggregate mesh kDOP. */
+	int32 GKDOPNodes = 0;
+	/** Number of leaf nodes in the aggregate mesh kDOP. */
+	int32 GKDOPNumLeaves = 0;
+	/** Maximum number of triangles per leaf node during the splitting process */
+	int32 GKDOPMaxTrisPerLeaf = DEFAULT_MAX_TRIS_PER_LEAF;
+	/** Total number of kDOP internal nodes traversed when tracing rays. */
+	volatile uint64 GKDOPParentNodesTraversed = 0;
+	/** Total number of kDOP leaf nodes traversed when tracing rays. */
+	volatile uint64 GKDOPLeafNodesTraversed = 0;
+	/** Total number of kDOP triangles tested when tracing rays. */
+	volatile uint64 GKDOPTrianglesTraversed = 0;
+	volatile uint64 GKDOPTrianglesTraversedReal = 0;
+	size_t kDOPPreallocatedMemory = 0u;
+	float kDOPBuildTime = 0.0f;
+
 	/** The list of nodes contained within this tree. Node 0 is always the root node. */
 	kDOPArray<NodeType, FRangeChecklessHeapAllocator> Nodes;
 
@@ -920,7 +1103,7 @@ struct TkDOPTree
 	 */
 	void Build(TArray<FkDOPBuildCollisionTriangle<KDOP_IDX_TYPE> >& BuildTriangles)
 	{
-		float kDOPBuildTime = 0;
+		kDOPBuildTime = 0;
 		{
 			FScopedRDTSCTimer kDOPBuildTimer(kDOPBuildTime);
 
@@ -933,7 +1116,7 @@ struct TkDOPTree
 
 			size_t NodesSize = (size_t)Nodes.GetTypeSize() * (size_t)Nodes.Max();
 			size_t SOATrianglesSize = (size_t)SOATriangles.GetTypeSize() * (size_t)SOATriangles.Max();
-			UE_LOG(LogLightmass, Log, TEXT("Preallocated %.1fGb for kDOP nodes and triangles"), (NodesSize + SOATrianglesSize) / 1024.0f / 1024.0f / 1024.0f);
+			kDOPPreallocatedMemory = NodesSize + SOATrianglesSize;
 
 			// Add the root node
 			Nodes.AddZeroed();
@@ -945,7 +1128,6 @@ struct TkDOPTree
 			Nodes.Shrink();
 			SOATriangles.Shrink();
 		}
-		UE_LOG(LogLightmass, Log, TEXT("Building kDOP took %5.2f seconds."),kDOPBuildTime);
 	}
 
 	
@@ -1168,6 +1350,11 @@ struct TkDOPTree
 		return bHit;
 	}
 
+	bool BoxCheck(TkDOPBoxCollisionCheck<COLL_DATA_PROVIDER, KDOP_IDX_TYPE>& Check) const
+	{
+		return Nodes[0].BoxCheck(Check);
+	}
+
 	/**
 	 * Dumps the kDOPTree 
 	 */
@@ -1361,4 +1548,17 @@ template <typename COLL_DATA_PROVIDER, typename KDOP_IDX_TYPE> struct TkDOPLineC
 	}
 };
 
+template <typename COLL_DATA_PROVIDER, typename KDOP_IDX_TYPE> struct TkDOPBoxCollisionCheck :
+	public TkDOPCollisionCheck<COLL_DATA_PROVIDER, KDOP_IDX_TYPE>
+{
+	const FBox Box;
+
+	TkDOPBoxCollisionCheck(
+		const FBox& InBox,
+		const COLL_DATA_PROVIDER& InCollDataProvider) :
+		TkDOPCollisionCheck<COLL_DATA_PROVIDER, KDOP_IDX_TYPE>(InCollDataProvider),
+		Box(InBox)
+	{}
+};
+
 } // namespace
diff --git a/Engine/Source/Programs/UnrealLightmass/Public/SceneExport.h b/Engine/Source/Programs/UnrealLightmass/Public/SceneExport.h
--- a/Engine/Source/Programs/UnrealLightmass/Public/SceneExport.h
+++ b/Engine/Source/Programs/UnrealLightmass/Public/SceneExport.h
@@ -72,6 +72,9 @@ public:
 	/** Whether to build Embree data structures for packet tracing. WIP feature - no lightmass algorithms emit packet tracing requests yet. */
 	bool bUseEmbreePacketTracing;
 
+	/** Whether to use kDOP trees & static mesh instancing to accelerate volumetric lightmap voxelization. Useful in scenes like large forest. */
+	bool bUseFastVoxelization;
+
 	/** 
 	 * Direct lighting, skylight radiosity and irradiance photons are cached on mapping surfaces to accelerate final gathering.
 	 * This controls the downsample factor for that cache, relative to the mapping's lightmap resolution.
