Implement bUseEmbreeInstancing (default to false). Instancing can greatly reduce memory consumption in scene like forests. Slows down ray traversal in small scale scenes like ParagonGPUTestbed by roughly 30% so default to false.

diff --git a/Engine/Config/BaseLightmass.ini b/Engine/Config/BaseLightmass.ini
--- a/Engine/Config/BaseLightmass.ini
+++ b/Engine/Config/BaseLightmass.ini
@@ -29,8 +29,10 @@ bRebuildDirtyGeometryForLighting=True
 bUseEmbree=true
 bVerifyEmbree=false
 bUseEmbreePacketTracing=false
-; Whether to use kDOP trees & static mesh instancing to accelerate volumetric lightmap voxelization. Useful in scenes like large forest
+; Whether to use kDOP trees to accelerate volumetric lightmap voxelization. Useful in scenes like large forest
 bUseFastVoxelization=false
+; Whether to use static mesh instancing to reduce memory consumption in scenes like large forest. Might slow down small scenes.
+bUseEmbreeInstancing=false
 bUseFilteredCubemapForSkylight=true
 MappingSurfaceCacheDownsampleFactor=2
 
diff --git a/Engine/Source/Editor/UnrealEd/Private/Lightmass/Lightmass.cpp b/Engine/Source/Editor/UnrealEd/Private/Lightmass/Lightmass.cpp
--- a/Engine/Source/Editor/UnrealEd/Private/Lightmass/Lightmass.cpp
+++ b/Engine/Source/Editor/UnrealEd/Private/Lightmass/Lightmass.cpp
@@ -2181,6 +2181,7 @@ void FLightmassExporter::WriteSceneSettings( Lightmass::FSceneFileHeader& Scene
 		Scene.GeneralSettings.bVerifyEmbree = Scene.GeneralSettings.bUseEmbree && bConfigBool;
 		VERIFYLIGHTMASSINI(GConfig->GetBool(TEXT("DevOptions.StaticLighting"), TEXT("bUseEmbreePacketTracing"), Scene.GeneralSettings.bUseEmbreePacketTracing, GLightmassIni));
 		VERIFYLIGHTMASSINI(GConfig->GetBool(TEXT("DevOptions.StaticLighting"), TEXT("bUseFastVoxelization"), Scene.GeneralSettings.bUseFastVoxelization, GLightmassIni));
+		VERIFYLIGHTMASSINI(GConfig->GetBool(TEXT("DevOptions.StaticLighting"), TEXT("bUseEmbreeInstancing"), Scene.GeneralSettings.bUseEmbreeInstancing, GLightmassIni));
 		VERIFYLIGHTMASSINI(GConfig->GetInt(TEXT("DevOptions.StaticLighting"), TEXT("MappingSurfaceCacheDownsampleFactor"), Scene.GeneralSettings.MappingSurfaceCacheDownsampleFactor, GLightmassIni));
 
 		int32 CheckQualityLevel;
diff --git a/Engine/Source/Programs/UnrealLightmass/Private/Lighting/AdaptiveVolumetricLightmap.cpp b/Engine/Source/Programs/UnrealLightmass/Private/Lighting/AdaptiveVolumetricLightmap.cpp
--- a/Engine/Source/Programs/UnrealLightmass/Private/Lighting/AdaptiveVolumetricLightmap.cpp
+++ b/Engine/Source/Programs/UnrealLightmass/Private/Lighting/AdaptiveVolumetricLightmap.cpp
@@ -142,14 +142,14 @@ bool FStaticLightingSystem::DoesVoxelIntersectSceneGeometry(const FBox& CellBoun
 			{
 				if (MeshInstance->Mapping->GetVolumeMapping() == nullptr)
 				{
-					if (MeshInstance->StaticMesh->VoxelizationMesh->IntersectBox(ExpandedCellBoundsSurfaceGeometry.TransformBy(MeshInstance->LocalToWorldInverseTranspose.GetTransposed())))
+					if (MeshInstance->StaticMesh->VoxelizationMesh->IntersectBox(ExpandedCellBoundsSurfaceGeometry.TransformBy(MeshInstance->WorldToLocal)))
 					{
 						return true;
 					}
 				}
 				else
 				{
-					if (MeshInstance->StaticMesh->VoxelizationMesh->IntersectBox(ExpandedCellBoundsVolumeGeometry.TransformBy(MeshInstance->LocalToWorldInverseTranspose.GetTransposed())))
+					if (MeshInstance->StaticMesh->VoxelizationMesh->IntersectBox(ExpandedCellBoundsVolumeGeometry.TransformBy(MeshInstance->WorldToLocal)))
 					{
 						return true;
 					}
diff --git a/Engine/Source/Programs/UnrealLightmass/Private/Lighting/Collision.cpp b/Engine/Source/Programs/UnrealLightmass/Private/Lighting/Collision.cpp
--- a/Engine/Source/Programs/UnrealLightmass/Private/Lighting/Collision.cpp
+++ b/Engine/Source/Programs/UnrealLightmass/Private/Lighting/Collision.cpp
@@ -126,8 +126,13 @@ void FDefaultAggregateMesh::AddMesh(const FStaticLightingMesh* Mesh, const FStat
 	}
 }
 
-void FDefaultAggregateMesh::AddMeshForVoxelization(const FStaticLightingMesh* Mesh, const FStaticLightingMapping* Mapping)
+void FDefaultAggregateMesh::AddMeshForVoxelization(const FStaticLightingMesh* Mesh, const FStaticLightingMapping* Mapping, bool bUseForInstancing)
 {
+	if (bUseForInstancing)
+	{
+		check(Mesh->GetInstanceableStaticMesh());
+	}
+
 	if (Mesh->LightingFlags & GI_INSTANCE_CASTSHADOW && Mesh->DoesMeshBelongToLOD0())
 	{
 		SceneBounds = SceneBounds + Mesh->BoundingBox;
@@ -145,19 +150,26 @@ void FDefaultAggregateMesh::AddMeshForVoxelization(const FStaticLightingMesh* Me
 		for (int32 TriangleIndex = 0; TriangleIndex < Mesh->NumTriangles; TriangleIndex++)
 		{
 			// Read the triangle from the mesh.
+			int32 I0 = 0;
+			int32 I1 = 0;
+			int32 I2 = 0;
 			FStaticLightingVertex V0;
 			FStaticLightingVertex V1;
 			FStaticLightingVertex V2;
 			int32 ElementIndex;
-			Mesh->GetTriangle(TriangleIndex, V0, V1, V2, ElementIndex);
-
-			if (Mesh->IsElementCastingShadow(ElementIndex) && !Mesh->IsTranslucent(ElementIndex))
+			if (bUseForInstancing)
+			{
+				Mesh->GetInstanceableStaticMesh()->GetNonTransformedTriangleIndices(TriangleIndex, I0, I1, I2);
+				Mesh->GetInstanceableStaticMesh()->GetNonTransformedTriangle(TriangleIndex, V0, V1, V2, ElementIndex);
+			}
+			else
 			{
-				int32 I0 = 0;
-				int32 I1 = 0;
-				int32 I2 = 0;
 				Mesh->GetTriangleIndices(TriangleIndex, I0, I1, I2);
+				Mesh->GetTriangle(TriangleIndex, V0, V1, V2, ElementIndex);
+			}
 
+			if (Mesh->IsElementCastingShadow(ElementIndex) && !Mesh->IsTranslucent(ElementIndex))
+			{
 				check(I0 <= Mesh->NumVertices && I1 <= Mesh->NumVertices && I2 <= Mesh->NumVertices);
 
 				const bool bTwoSided = Mesh->IsTwoSided(ElementIndex) || Mesh->IsCastingShadowAsTwoSided();
diff --git a/Engine/Source/Programs/UnrealLightmass/Private/Lighting/Collision.h b/Engine/Source/Programs/UnrealLightmass/Private/Lighting/Collision.h
--- a/Engine/Source/Programs/UnrealLightmass/Private/Lighting/Collision.h
+++ b/Engine/Source/Programs/UnrealLightmass/Private/Lighting/Collision.h
@@ -230,7 +230,7 @@ public:
 
 	virtual void AddMesh(const FStaticLightingMesh* Mesh, const FStaticLightingMapping* Mapping) override;
 
-	void AddMeshForVoxelization(const FStaticLightingMesh* Mesh, const FStaticLightingMapping* Mapping);
+	void AddMeshForVoxelization(const FStaticLightingMesh* Mesh, const FStaticLightingMapping* Mapping, bool bUseForInstancing = false);
 
 	virtual void ReserveMemory( int32 NumMeshes, int32 NumVertices, int32 NumTriangles ) override;
 
diff --git a/Engine/Source/Programs/UnrealLightmass/Private/Lighting/Embree.cpp b/Engine/Source/Programs/UnrealLightmass/Private/Lighting/Embree.cpp
--- a/Engine/Source/Programs/UnrealLightmass/Private/Lighting/Embree.cpp
+++ b/Engine/Source/Programs/UnrealLightmass/Private/Lighting/Embree.cpp
@@ -113,12 +113,28 @@ struct FEmbreeFilterProcessor
 	FEmbreeFilterProcessor(FEmbreeRay& InRay, const FEmbreeGeometry& InGeo) : Ray(InRay), Geo(InGeo), bCoordsDirty(true) 
 	{
 		Desc = Geo.TriangleDescs[Ray.primID];
+		Mesh = Geo.Mesh;
+		if (Ray.instID != -1)
+		{
+			// if instancing is used, material evaluation is deferred here
+			Mesh = Geo.ParentAggregateMesh->StaticMeshInstancesToMappings[Ray.instID]->Mesh;
+			int32 ElementIndex = Desc.ElementIndex;
+			Desc.CastShadow = Mesh->IsElementCastingShadow(ElementIndex);
+			Desc.StaticAndOpaqueMask = !Mesh->IsMasked(ElementIndex) && !Mesh->IsTranslucent(ElementIndex) && !Mesh->bMovable;;
+			Desc.TwoSidedMask = Mesh->IsTwoSided(ElementIndex) || Mesh->IsCastingShadowAsTwoSided();
+			Desc.Translucent = Mesh->IsTranslucent(ElementIndex);
+			Desc.SurfaceDomain = Mesh->IsSurfaceDomain(ElementIndex);
+			Desc.IndirectlyShadowedOnly = Mesh->IsIndirectlyShadowedOnly(ElementIndex);
+			Desc.Masked = Mesh->IsMasked(ElementIndex);
+			Desc.CastShadowAsMasked = Mesh->IsCastingShadowsAsMasked(ElementIndex);
+		}
 		s = 1 - Ray.u - Ray.v;
 	}
 
 	FEmbreeRay& Ray;
 	const FEmbreeGeometry& Geo;
 	FEmbreeTriangleDesc Desc;
+	const FStaticLightingMesh* Mesh;
 
 	float s; // (s,u,v) : Barycentric Weights
 	int32 Index0;
@@ -187,7 +203,7 @@ struct FEmbreeFilterProcessor
 	EMBREE_INLINE bool HitRejectedByHLODTest() const
 	{
 		const uint32 InvalidIndex = 0xFFFF;
-		uint32 GeoHLODTreeIndex = (Geo.Mesh->GetLODIndices() & 0xFFFF0000) >> 16;
+		uint32 GeoHLODTreeIndex = (Mesh->GetLODIndices() & 0xFFFF0000) >> 16;
 		uint32 RayHLODTreeIndex = Ray.MappingMesh ? (Ray.MappingMesh->GetLODIndices() & 0xFFFF0000) >> 16 : InvalidIndex;
 
 		bool bRejectHit = false;
@@ -195,7 +211,7 @@ struct FEmbreeFilterProcessor
 		// If either Geo or Ray is a HLOD (0xFFFF being invalid HLOD)
 		if (GeoHLODTreeIndex != InvalidIndex || RayHLODTreeIndex != InvalidIndex)
 		{
-			uint32 GeoHLODRange = Geo.Mesh->GetHLODRange();
+			uint32 GeoHLODRange = Mesh->GetHLODRange();
 			uint32 GeoHLODRangeStart = GeoHLODRange & 0xFFFF;
 			uint32 GeoHLODRangeEnd = (GeoHLODRange & 0xFFFF0000) >> 16;
 
@@ -233,13 +249,13 @@ struct FEmbreeFilterProcessor
 
 	EMBREE_INLINE bool HitRejectedByLODIndexTest() const
 	{
-		uint32 GeoMeshLODIndex = Geo.Mesh->GetLODIndices() & 0xFFFF;
+		uint32 GeoMeshLODIndex = Mesh->GetLODIndices() & 0xFFFF;
 
 		// Only shadows from appropriate mesh LODs.
 		if (Ray.MappingMesh)
 		{
 			// If it is not from the same mesh, then only LOD 0 can cast shadow.
-			if (Ray.MappingMesh->MeshIndex != Geo.Mesh->MeshIndex)
+			if (Ray.MappingMesh->MeshIndex != Mesh->MeshIndex)
 			{
 				return GeoMeshLODIndex != 0;
 			}
@@ -257,9 +273,9 @@ struct FEmbreeFilterProcessor
 	EMBREE_INLINE bool HitRejectedBySelfShadowTest() const
 	{
 		// No self shadows, or only self shadow
-		if ((Geo.Mesh == Ray.ShadowMesh && ((Geo.Mesh->LightingFlags & GI_INSTANCE_SELFSHADOWDISABLE) || (Ray.TraceFlags & LIGHTRAY_SELFSHADOWDISABLE)))
+		if ((Mesh == Ray.ShadowMesh && ((Mesh->LightingFlags & GI_INSTANCE_SELFSHADOWDISABLE) || (Ray.TraceFlags & LIGHTRAY_SELFSHADOWDISABLE)))
 		|| (Ray.bDirectShadowingRay && Desc.IndirectlyShadowedOnly)
-		|| (Geo.Mesh != Ray.ShadowMesh && (Geo.Mesh->LightingFlags & GI_INSTANCE_SELFSHADOWONLY)))
+		|| (Mesh != Ray.ShadowMesh && (Mesh->LightingFlags & GI_INSTANCE_SELFSHADOWONLY)))
 		{
 			return true;
 		}
@@ -272,7 +288,7 @@ struct FEmbreeFilterProcessor
 
 		if (Desc.Masked || (Ray.bDirectShadowingRay && Desc.CastShadowAsMasked))
 		{
- 			return !Geo.Mesh->EvaluateMaskedCollision(TextureCoordinates, Desc.ElementIndex);
+ 			return !Mesh->EvaluateMaskedCollision(TextureCoordinates, Desc.ElementIndex);
 		}
 		return false;
 	}
@@ -282,7 +298,14 @@ void FEmbreeFilterProcessor::UpdateCoordinates()
 {
 	if (bCoordsDirty)
 	{
-		Geo.Mesh->GetTriangleIndices(Ray.primID, Index0, Index1, Index2);
+		if (Ray.instID == -1)
+		{
+			Mesh->GetTriangleIndices(Ray.primID, Index0, Index1, Index2);
+		}
+		else
+		{
+			Mesh->GetInstanceableStaticMesh()->GetNonTransformedTriangleIndices(Ray.primID, Index0, Index1, Index2);
+		}
 
 		const FVector2D& UV1 = Geo.UVs[Index0];
 		const FVector2D& UV2 = Geo.UVs[Index1];
@@ -315,7 +338,14 @@ void FEmbreeFilterProcessor::UpdateRay()
 	}
 	else if (bCoordsDirty)
 	{
-		Geo.Mesh->GetTriangleIndices(Ray.primID, Index0, Index1, Index2);
+		if (Ray.instID == -1)
+		{
+			Mesh->GetTriangleIndices(Ray.primID, Index0, Index1, Index2);
+		}
+		else
+		{
+			Mesh->GetInstanceableStaticMesh()->GetNonTransformedTriangleIndices(Ray.primID, Index0, Index1, Index2);
+		}
 	}
 
 	// Transmission : updated outside of this scope.
@@ -359,7 +389,7 @@ void EmbreeFilterFunc(void* UserPtr, RTCRay& InRay)
 
 			// Accumulate the total transmission along the ray
 			// The result is order independent so the intersections don't have to be strictly front to back
-			Proc.Ray.TransmissionAcc.Push(Proc.Geo.Mesh->EvaluateTransmission(Proc.TextureCoordinates, Proc.Desc.ElementIndex), Proc.Ray.tfar);
+			Proc.Ray.TransmissionAcc.Push(Proc.Mesh->EvaluateTransmission(Proc.TextureCoordinates, Proc.Desc.ElementIndex), Proc.Ray.tfar);
 
 		}
 		Proc.Invalidate();
@@ -409,7 +439,8 @@ FEmbreeGeometry::FEmbreeGeometry(
 	RTCScene EmbreeScene, 
 	const FBoxSphereBounds& ImportanceBounds,
 	const FStaticLightingMesh* InMesh,
-	const FStaticLightingMapping* InMapping
+	const FStaticLightingMapping* InMapping,
+	bool bUseForInstancing
 	) :
 	Mesh(InMesh),
 	Mapping(InMapping),
@@ -417,6 +448,11 @@ FEmbreeGeometry::FEmbreeGeometry(
 	SurfaceAreaWithinImportanceVolume(0),
 	bHasShadowCastingPrimitives(false)
 {
+	if (bUseForInstancing)
+	{
+		check(Mesh->GetInstanceableStaticMesh() != nullptr);
+	}
+
 	const FStaticLightingTextureMapping* TextureMapping = Mapping ? Mapping->GetTextureMapping() : NULL;
 
 	GeomID = rtcNewTriangleMesh(EmbreeScene, RTC_GEOMETRY_STATIC, Mesh->NumTriangles, Mesh->NumVertices);
@@ -433,8 +469,16 @@ FEmbreeGeometry::FEmbreeGeometry(
 		FStaticLightingVertex V0, V1, V2;
 		int32 ElementIndex;
 
-		Mesh->GetTriangleIndices(TriangleIndex,I0,I1,I2);
-		Mesh->GetTriangle(TriangleIndex,V0,V1,V2,ElementIndex);
+		if (bUseForInstancing)
+		{
+			Mesh->GetInstanceableStaticMesh()->GetNonTransformedTriangleIndices(TriangleIndex, I0, I1, I2);
+			Mesh->GetInstanceableStaticMesh()->GetNonTransformedTriangle(TriangleIndex, V0, V1, V2, ElementIndex);
+		}
+		else
+		{
+			Mesh->GetTriangleIndices(TriangleIndex, I0, I1, I2);
+			Mesh->GetTriangle(TriangleIndex, V0, V1, V2, ElementIndex);
+		}
 
 		// Compute the triangle's normal.
 		const FVector4 TriangleNormal = (V2.WorldPosition - V0.WorldPosition) ^ (V1.WorldPosition - V0.WorldPosition);
@@ -443,16 +487,23 @@ FEmbreeGeometry::FEmbreeGeometry(
 
 		FEmbreeTriangleDesc& Desc = TriangleDescs[TriangleIndex];
 		Desc.ElementIndex = ElementIndex;
-		Desc.CastShadow = TriangleArea > TRIANGLE_AREA_THRESHOLD && Mesh->IsElementCastingShadow(ElementIndex);
-		Desc.StaticAndOpaqueMask = !Mesh->IsMasked(ElementIndex) && !Mesh->IsTranslucent(ElementIndex) && !Mesh->bMovable;;
-		Desc.TwoSidedMask = Mesh->IsTwoSided(ElementIndex) || Mesh->IsCastingShadowAsTwoSided();
-		Desc.Translucent = Mesh->IsTranslucent(ElementIndex);
-		Desc.SurfaceDomain = Mesh->IsSurfaceDomain(ElementIndex);
-		Desc.IndirectlyShadowedOnly = Mesh->IsIndirectlyShadowedOnly(ElementIndex);
-		Desc.Masked = Mesh->IsMasked(ElementIndex);
-		Desc.CastShadowAsMasked = Mesh->IsCastingShadowsAsMasked(ElementIndex);
+		Desc.CastShadow = false;
+
+		if (!bUseForInstancing)
+		{
+			// When instancing is not used, evaluate material properties here
+			// Otherwise, defer material evaluation until ray intersection
+			Desc.CastShadow = TriangleArea > TRIANGLE_AREA_THRESHOLD && Mesh->IsElementCastingShadow(ElementIndex);
+			Desc.StaticAndOpaqueMask = !Mesh->IsMasked(ElementIndex) && !Mesh->IsTranslucent(ElementIndex) && !Mesh->bMovable;;
+			Desc.TwoSidedMask = Mesh->IsTwoSided(ElementIndex) || Mesh->IsCastingShadowAsTwoSided();
+			Desc.Translucent = Mesh->IsTranslucent(ElementIndex);
+			Desc.SurfaceDomain = Mesh->IsSurfaceDomain(ElementIndex);
+			Desc.IndirectlyShadowedOnly = Mesh->IsIndirectlyShadowedOnly(ElementIndex);
+			Desc.Masked = Mesh->IsMasked(ElementIndex);
+			Desc.CastShadowAsMasked = Mesh->IsCastingShadowsAsMasked(ElementIndex);
+		}
 
-		if (Desc.CastShadow)
+		if (TriangleArea > TRIANGLE_AREA_THRESHOLD && (bUseForInstancing || Desc.CastShadow))
 		{
 			Indices[TriangleIndex * 3 + 0] = I0;
 			Indices[TriangleIndex * 3 + 1] = I1;
@@ -480,16 +531,19 @@ FEmbreeGeometry::FEmbreeGeometry(
 			LightmapUVs[I2] = V2.TextureCoordinates[TextureMapping->LightmapTextureCoordinateIndex];
 		}
 
-		SurfaceArea += TriangleArea;
-
-		// Sum the total triangle area of everything in the aggregate mesh within the importance volume,
-		// if any vertex is contained or if there is no importance volume at all
-		if( ImportanceBounds.SphereRadius < DELTA ||
-			ImportanceBounds.GetBox().IsInside( V0.WorldPosition ) ||
-			ImportanceBounds.GetBox().IsInside( V1.WorldPosition ) ||
-			ImportanceBounds.GetBox().IsInside( V2.WorldPosition ) )
+		if (!bUseForInstancing)
 		{
-			SurfaceAreaWithinImportanceVolume += TriangleArea;
+			SurfaceArea += TriangleArea;
+
+			// Sum the total triangle area of everything in the aggregate mesh within the importance volume,
+			// if any vertex is contained or if there is no importance volume at all
+			if (ImportanceBounds.SphereRadius < DELTA ||
+				ImportanceBounds.GetBox().IsInside(V0.WorldPosition) ||
+				ImportanceBounds.GetBox().IsInside(V1.WorldPosition) ||
+				ImportanceBounds.GetBox().IsInside(V2.WorldPosition))
+			{
+				SurfaceAreaWithinImportanceVolume += TriangleArea;
+			}
 		}
 
 	}
@@ -499,6 +553,39 @@ FEmbreeGeometry::FEmbreeGeometry(
 	check(rtcDeviceGetError(EmbreeDevice) == RTC_NO_ERROR);
 }
 
+void CalculateSurfaceArea(const FStaticLightingMesh* Mesh, const FBoxSphereBounds& ImportanceBounds, float& SurfaceArea, float& SurfaceAreaWithinImportanceVolume)
+{
+	SurfaceArea = 0.0f;
+	SurfaceAreaWithinImportanceVolume = 0.0f;
+
+	for (int32 TriangleIndex = 0; TriangleIndex < Mesh->NumTriangles; TriangleIndex++)
+	{
+		int32 I0 = 0, I1 = 0, I2 = 0;
+		FStaticLightingVertex V0, V1, V2;
+		int32 ElementIndex;
+
+		Mesh->GetTriangleIndices(TriangleIndex, I0, I1, I2);
+		Mesh->GetTriangle(TriangleIndex, V0, V1, V2, ElementIndex);
+
+		// Compute the triangle's normal.
+		const FVector4 TriangleNormal = (V2.WorldPosition - V0.WorldPosition) ^ (V1.WorldPosition - V0.WorldPosition);
+		// Compute the triangle area.
+		const float TriangleArea = TriangleNormal.Size3() * 0.5f;
+
+		SurfaceArea += TriangleArea;
+
+		// Sum the total triangle area of everything in the aggregate mesh within the importance volume,
+		// if any vertex is contained or if there is no importance volume at all
+		if (ImportanceBounds.SphereRadius < DELTA ||
+			ImportanceBounds.GetBox().IsInside(V0.WorldPosition) ||
+			ImportanceBounds.GetBox().IsInside(V1.WorldPosition) ||
+			ImportanceBounds.GetBox().IsInside(V2.WorldPosition))
+		{
+			SurfaceAreaWithinImportanceVolume += TriangleArea;
+		}
+	}
+}
+
 bool EmbreeMemoryMonitor(const ssize_t bytes, const bool post)
 {
 	FPlatformAtomics::InterlockedAdd(&GEmbreeAllocatedSpace, bytes);
@@ -509,7 +596,8 @@ FEmbreeAggregateMesh::FEmbreeAggregateMesh(const FScene& InScene):
 	FStaticLightingAggregateMesh(InScene),
 	EmbreeDevice(NULL),
 	EmbreeScene(NULL),
-	TotalNumTriangles(0)
+	TotalNumTriangles(0),
+	TotalNumTrianglesInstanced(0)
 {
 	rtcDeviceSetMemoryMonitorFunction(InScene.EmbreeDevice, EmbreeMemoryMonitor);
 
@@ -548,28 +636,80 @@ FEmbreeAggregateMesh::~FEmbreeAggregateMesh()
 */
 void FEmbreeAggregateMesh::AddMesh(const FStaticLightingMesh* Mesh, const FStaticLightingMapping* Mapping)
 {
-	// Only use shadow casting meshes.
-	if( Mesh->LightingFlags&GI_INSTANCE_CASTSHADOW )
-	{
-		SceneBounds = SceneBounds + Mesh->BoundingBox;
+ 	// Only use shadow casting meshes.
+ 	if( Mesh->LightingFlags&GI_INSTANCE_CASTSHADOW )
+ 	{
+ 		SceneBounds = SceneBounds + Mesh->BoundingBox;
 
-		FEmbreeGeometry* Geo = new FEmbreeGeometry(EmbreeDevice, EmbreeScene, Scene.GetImportanceBounds(), Mesh, Mapping);
-		MeshInfos.Add(Geo);
-		
-		rtcSetUserData(EmbreeScene, Geo->GeomID, Geo);
-		rtcSetIntersectionFilterFunction(EmbreeScene, Geo->GeomID, EmbreeFilterFunc);
-		rtcSetOcclusionFilterFunction(EmbreeScene, Geo->GeomID, EmbreeFilterFunc);
+		if (Scene.GeneralSettings.bUseEmbreeInstancing && Mesh->GetInstanceableStaticMesh() != nullptr)
+		{
+			const FStaticMeshStaticLightingMesh* StaticMeshInstance = Mesh->GetInstanceableStaticMesh();
+
+			const FStaticMeshLOD* LOD = &StaticMeshInstance->StaticMesh->GetLOD(StaticMeshInstance->GetMeshLODIndex());
 
-		rtcSetIntersectionFilterFunction4(EmbreeScene, Geo->GeomID, EmbreeFilterFunc4);
-		rtcSetOcclusionFilterFunction4(EmbreeScene, Geo->GeomID, EmbreeFilterFunc4);
+			if (!StaticMeshGeometries.Find(LOD))
+			{
+				RTCScene MeshScene = rtcDeviceNewScene(EmbreeDevice, RTC_SCENE_STATIC, RTC_INTERSECT1);
 
-		bHasShadowCastingPrimitives |= Geo->bHasShadowCastingPrimitives;
+				FEmbreeGeometry* Geo = new FEmbreeGeometry(EmbreeDevice, MeshScene, Scene.GetImportanceBounds(), Mesh, Mapping, true);
+				Geo->ParentAggregateMesh = this;
+				StaticMeshGeometries.Add(LOD, FEmbreeStaticMeshGeometry { MeshScene, Geo });
+				MeshInfos.Add(Geo);
 
-		// Sum the total triangle area of everything in the aggregate mesh
-		SceneSurfaceArea += Geo->SurfaceArea;
-		SceneSurfaceAreaWithinImportanceVolume += Geo->SurfaceAreaWithinImportanceVolume;
+				rtcSetIntersectionFilterFunction(MeshScene, Geo->GeomID, EmbreeFilterFunc);
+				rtcSetOcclusionFilterFunction(MeshScene, Geo->GeomID, EmbreeFilterFunc);
 
-		TotalNumTriangles += Mesh->NumTriangles;
+				rtcSetIntersectionFilterFunction4(MeshScene, Geo->GeomID, EmbreeFilterFunc4);
+				rtcSetOcclusionFilterFunction4(MeshScene, Geo->GeomID, EmbreeFilterFunc4);
+
+				rtcSetUserData(MeshScene, Geo->GeomID, Geo);
+
+				rtcCommit(MeshScene);
+
+				bHasShadowCastingPrimitives |= Geo->bHasShadowCastingPrimitives;
+
+				TotalNumTriangles += Mesh->NumTriangles;
+			}
+			else
+			{
+				TotalNumTrianglesInstanced += Mesh->NumTriangles;
+			}
+
+			// Sum the total triangle area of everything in the aggregate mesh
+			float SurfaceArea, SurfaceAreaWithinImportanceVolume;
+			CalculateSurfaceArea(Mesh, Scene.GetImportanceBounds(), SurfaceArea, SurfaceAreaWithinImportanceVolume);
+			SceneSurfaceArea += SurfaceArea;
+			SceneSurfaceAreaWithinImportanceVolume += SurfaceAreaWithinImportanceVolume;
+
+			auto InstID = rtcNewInstance(EmbreeScene, StaticMeshGeometries[LOD].MeshScene);
+			rtcSetUserData(EmbreeScene, InstID, (void*)Mapping);
+			rtcSetTransform(EmbreeScene, InstID, RTC_MATRIX_COLUMN_MAJOR_ALIGNED16, (const float*)StaticMeshInstance->LocalToWorld.M);
+			if (StaticMeshInstancesToMappings.Num() < (int32)InstID + 1)
+			{
+				StaticMeshInstancesToMappings.SetNum(InstID + 1);
+			}
+			StaticMeshInstancesToMappings[InstID] = Mapping;
+
+		}
+		else
+		{
+			FEmbreeGeometry* Geo = new FEmbreeGeometry(EmbreeDevice, EmbreeScene, Scene.GetImportanceBounds(), Mesh, Mapping, false);
+			MeshInfos.Add(Geo);
+
+			rtcSetUserData(EmbreeScene, Geo->GeomID, Geo);
+			rtcSetIntersectionFilterFunction(EmbreeScene, Geo->GeomID, EmbreeFilterFunc);
+			rtcSetOcclusionFilterFunction(EmbreeScene, Geo->GeomID, EmbreeFilterFunc);
+
+			rtcSetIntersectionFilterFunction4(EmbreeScene, Geo->GeomID, EmbreeFilterFunc4);
+			rtcSetOcclusionFilterFunction4(EmbreeScene, Geo->GeomID, EmbreeFilterFunc4);
+
+			bHasShadowCastingPrimitives |= Geo->bHasShadowCastingPrimitives;
+
+			// Sum the total triangle area of everything in the aggregate mesh
+			SceneSurfaceArea += Geo->SurfaceArea;
+			SceneSurfaceAreaWithinImportanceVolume += Geo->SurfaceAreaWithinImportanceVolume;
+			TotalNumTriangles += Mesh->NumTriangles;
+		}
 	}
 }
 
@@ -599,7 +739,14 @@ void FEmbreeAggregateMesh::DumpStats() const
 
 	UE_LOG(LogLightmass, Log, TEXT("\n"));
 	UE_LOG(LogLightmass, Log, TEXT("Collision Mesh Overview:"));
-	UE_LOG(LogLightmass, Log, TEXT("Num Triangles         : %d"), TotalNumTriangles);
+	if (Scene.GeneralSettings.bUseEmbreeInstancing)
+	{
+		UE_LOG(LogLightmass, Log, TEXT("Num Triangles         : %d (Instanced to %d)"), TotalNumTriangles, TotalNumTriangles + TotalNumTrianglesInstanced);
+	}
+	else
+	{
+		UE_LOG(LogLightmass, Log, TEXT("Num Triangles         : %d"), TotalNumTriangles);
+	}
 	UE_LOG(LogLightmass, Log, TEXT("MeshInfos             : %7.1fMb"), MeshInfoSize / 1048576.0f);
 	UE_LOG(LogLightmass, Log, TEXT("UVs                   : %7.1fMb"), UVSize / 1048576.0f);
 	UE_LOG(LogLightmass, Log, TEXT("LightmapUVs           : %7.1fMb"), LightmapUV / 1048576.0f);
@@ -651,17 +798,25 @@ bool FEmbreeAggregateMesh::IntersectLightRay(
 
 	if (EmbreeRay.geomID != -1 && EmbreeRay.primID != -1)
 	{
-		const FEmbreeGeometry& Geo = *(FEmbreeGeometry*)rtcGetUserData(EmbreeScene, EmbreeRay.geomID);
-
 		FMinimalStaticLightingVertex EmbreeVertex;
 		EmbreeVertex.WorldPosition = LightRay.Start + LightRay.Direction * EmbreeRay.tfar;
-		EmbreeVertex.WorldTangentZ = FVector(EmbreeRay.Ng[0], EmbreeRay.Ng[1], EmbreeRay.Ng[2]).GetSafeNormal();
 
 		EmbreeVertex.TextureCoordinates[0] = EmbreeRay.TextureCoordinates;
 		EmbreeVertex.TextureCoordinates[1] = EmbreeRay.LightmapCoordinates;
 
-		ClosestIntersection = FLightRayIntersection(true, EmbreeVertex, Geo.Mesh, Geo.Mapping, EmbreeRay.ElementIndex);
-
+		if (EmbreeRay.instID == -1)
+		{
+			const FEmbreeGeometry& Geo = *(FEmbreeGeometry*)rtcGetUserData(EmbreeScene, EmbreeRay.geomID);
+			EmbreeVertex.WorldTangentZ = FVector(EmbreeRay.Ng[0], EmbreeRay.Ng[1], EmbreeRay.Ng[2]).GetSafeNormal();
+			ClosestIntersection = FLightRayIntersection(true, EmbreeVertex, Geo.Mesh, Geo.Mapping, EmbreeRay.ElementIndex);
+		}
+		else
+		{
+			FStaticLightingMapping* Mapping = (FStaticLightingMapping*)rtcGetUserData(EmbreeScene, EmbreeRay.instID);
+			FVector GeometryNormal(EmbreeRay.Ng[0], EmbreeRay.Ng[1], EmbreeRay.Ng[2]);
+			EmbreeVertex.WorldTangentZ = Mapping->Mesh->GetInstanceableStaticMesh()->LocalToWorldInverseTranspose.TransformVector(GeometryNormal).GetSafeNormal();
+			ClosestIntersection = FLightRayIntersection(true, EmbreeVertex, Mapping->Mesh, Mapping, EmbreeRay.ElementIndex);
+		}
 		EmbreeRay.TransmissionAcc.Resolve(ClosestIntersection.Transmission, EmbreeRay.tfar);
 	}
 	else
diff --git a/Engine/Source/Programs/UnrealLightmass/Private/Lighting/Embree.h b/Engine/Source/Programs/UnrealLightmass/Private/Lighting/Embree.h
--- a/Engine/Source/Programs/UnrealLightmass/Private/Lighting/Embree.h
+++ b/Engine/Source/Programs/UnrealLightmass/Private/Lighting/Embree.h
@@ -164,7 +164,8 @@ namespace Lightmass
 			RTCScene EmbreeScene, 
 			const FBoxSphereBounds& ImportanceBounds,
 			const FStaticLightingMesh* InMesh,
-			const FStaticLightingMapping* InMapping
+			const FStaticLightingMapping* InMapping,
+			bool bUseForInstancing
 			);
 
 		const FStaticLightingMesh* Mesh;
@@ -179,6 +180,8 @@ namespace Lightmass
 		float SurfaceArea;
 		float SurfaceAreaWithinImportanceVolume;
 		bool bHasShadowCastingPrimitives;
+
+		class FEmbreeAggregateMesh* ParentAggregateMesh;
 	};
 
 	class FEmbreeAggregateMesh final : public FStaticLightingAggregateMesh
@@ -215,11 +218,20 @@ namespace Lightmass
 			FCoherentRayCache& CoherentRayCache,
 			FLightRayIntersection* ClosestIntersections) const override;
 
+		TArray<const FStaticLightingMapping*> StaticMeshInstancesToMappings;
+
 	private:
 
 		/** Information about the meshes used in the kDOP tree. */
 		TArray<const FEmbreeGeometry*> MeshInfos;
 
+		struct FEmbreeStaticMeshGeometry
+		{
+			RTCScene MeshScene;
+			FEmbreeGeometry* Geo;
+		};
+		TMap<const FStaticMeshLOD*, FEmbreeStaticMeshGeometry> StaticMeshGeometries;
+
 		/** Embree scene */
 		RTCDevice EmbreeDevice;
 		RTCScene EmbreeScene;
@@ -227,6 +239,8 @@ namespace Lightmass
 		/** Total number of triangles in the shadow mesh */
 		int32 TotalNumTriangles;
 
+		/** Total number of instanced triangles in the shadow mesh */
+		int32 TotalNumTrianglesInstanced;
 	};
 
 	class FEmbreeVerifyAggregateMesh final : public FStaticLightingAggregateMesh
diff --git a/Engine/Source/Programs/UnrealLightmass/Private/Lighting/LightingSystem.cpp b/Engine/Source/Programs/UnrealLightmass/Private/Lighting/LightingSystem.cpp
--- a/Engine/Source/Programs/UnrealLightmass/Private/Lighting/LightingSystem.cpp
+++ b/Engine/Source/Programs/UnrealLightmass/Private/Lighting/LightingSystem.cpp
@@ -440,20 +440,8 @@ FStaticLightingSystem::FStaticLightingSystem(const FLightingBuildOptions& InOpti
 					if (MeshInstance->LightingFlags & GI_INSTANCE_CASTSHADOW && MeshInstance->DoesMeshBelongToLOD0())
 					{
 						MeshInstance->StaticMesh->VoxelizationMesh = new FDefaultAggregateMesh(Scene);
-
-						{
-							auto LocalToWorld = MeshInstance->LocalToWorld;
-							auto LocalToWorldInverseTranspose = MeshInstance->LocalToWorldInverseTranspose;
-
-							MeshInstance->LocalToWorld = FMatrix::Identity;
-							MeshInstance->LocalToWorldInverseTranspose = FMatrix::Identity;
-
-							MeshInstance->StaticMesh->VoxelizationMesh->AddMeshForVoxelization(MeshInstance, MeshInstance->Mapping);
-							MeshInstance->StaticMesh->VoxelizationMesh->PrepareForRaytracing();
-
-							MeshInstance->LocalToWorld = LocalToWorld;
-							MeshInstance->LocalToWorldInverseTranspose = LocalToWorldInverseTranspose;
-						}
+						MeshInstance->StaticMesh->VoxelizationMesh->AddMeshForVoxelization(MeshInstance, MeshInstance->Mapping, true);
+						MeshInstance->StaticMesh->VoxelizationMesh->PrepareForRaytracing();
 					}
 				}
 			}
diff --git a/Engine/Source/Programs/UnrealLightmass/Private/Lighting/StaticMesh.cpp b/Engine/Source/Programs/UnrealLightmass/Private/Lighting/StaticMesh.cpp
--- a/Engine/Source/Programs/UnrealLightmass/Private/Lighting/StaticMesh.cpp
+++ b/Engine/Source/Programs/UnrealLightmass/Private/Lighting/StaticMesh.cpp
@@ -232,6 +232,36 @@ void FStaticMeshStaticLightingMesh::GetTriangle(int32 TriangleIndex,FStaticLight
 	check(ElementIndex >= 0);
 }
 
+void FStaticMeshStaticLightingMesh::GetNonTransformedTriangle(int32 TriangleIndex, FStaticLightingVertex& OutV0, FStaticLightingVertex& OutV1, FStaticLightingVertex& OutV2, int32& ElementIndex) const
+{
+	check(!bIsSplineMesh);
+
+	const FStaticMeshLOD& LODRenderData = StaticMesh->GetLOD(GetMeshLODIndex());
+
+	// Lookup the triangle's vertex indices.
+	const uint32 I0 = LODRenderData.GetIndex(TriangleIndex * 3 + 0);
+	const uint32 I1 = LODRenderData.GetIndex(TriangleIndex * 3 + 1);
+	const uint32 I2 = LODRenderData.GetIndex(TriangleIndex * 3 + 2);
+
+	// Translate the triangle's static mesh vertices to static lighting vertices.
+	GetStaticLightingVertex(LODRenderData.GetVertex(I0), FMatrix::Identity, FMatrix::Identity, false, SplineParameters, OutV0);
+	GetStaticLightingVertex(LODRenderData.GetVertex(I1), FMatrix::Identity, FMatrix::Identity, false, SplineParameters, OutV1);
+	GetStaticLightingVertex(LODRenderData.GetVertex(I2), FMatrix::Identity, FMatrix::Identity, false, SplineParameters, OutV2);
+
+	const FStaticMeshLOD& MeshLOD = StaticMesh->GetLOD(GetMeshLODIndex());
+	ElementIndex = INDEX_NONE;
+	for (uint32 MeshElementIndex = 0; MeshElementIndex < MeshLOD.NumElements; MeshElementIndex++)
+	{
+		const FStaticMeshElement& CurrentElement = MeshLOD.GetElement(MeshElementIndex);
+		if ((uint32)TriangleIndex >= CurrentElement.FirstIndex / 3 && (uint32)TriangleIndex < CurrentElement.FirstIndex / 3 + CurrentElement.NumTriangles)
+		{
+			ElementIndex = MeshElementIndex;
+			break;
+		}
+	}
+	check(ElementIndex >= 0);
+}
+
 void FStaticMeshStaticLightingMesh::GetTriangleIndices(int32 TriangleIndex,int32& OutI0,int32& OutI1,int32& OutI2) const
 {
  	const FStaticMeshLOD& LODRenderData = StaticMesh->GetLOD(GetMeshLODIndex());
@@ -242,6 +272,16 @@ void FStaticMeshStaticLightingMesh::GetTriangleIndices(int32 TriangleIndex,int32
  	OutI2 = LODRenderData.GetIndex(TriangleIndex * 3 + (bReverseWinding ? 1 : 2));
 }
 
+void FStaticMeshStaticLightingMesh::GetNonTransformedTriangleIndices(int32 TriangleIndex, int32& OutI0, int32& OutI1, int32& OutI2) const
+{
+	const FStaticMeshLOD& LODRenderData = StaticMesh->GetLOD(GetMeshLODIndex());
+
+	// Lookup the triangle's vertex indices.
+	OutI0 = LODRenderData.GetIndex(TriangleIndex * 3 + 0);
+	OutI1 = LODRenderData.GetIndex(TriangleIndex * 3 + 1);
+	OutI2 = LODRenderData.GetIndex(TriangleIndex * 3 + 2);
+}
+
 bool FStaticMeshStaticLightingMesh::IsElementCastingShadow(int32 ElementIndex) const
 {
 	const FStaticMeshLOD& LODRenderData = StaticMesh->GetLOD(GetMeshLODIndex());
@@ -268,6 +308,7 @@ void FStaticMeshStaticLightingMesh::Import( FLightmassImporter& Importer )
 	SplineParameters = TempSMSLMD.SplineParameters;
 
 	// calculate the inverse transpose
+	WorldToLocal = LocalToWorld.InverseFast();
 	LocalToWorldInverseTranspose = LocalToWorld.InverseFast().GetTransposed();
 
 	// we have the guid for the mesh, now hook it up to the actual static mesh
diff --git a/Engine/Source/Programs/UnrealLightmass/Private/Lighting/StaticMesh.h b/Engine/Source/Programs/UnrealLightmass/Private/Lighting/StaticMesh.h
--- a/Engine/Source/Programs/UnrealLightmass/Private/Lighting/StaticMesh.h
+++ b/Engine/Source/Programs/UnrealLightmass/Private/Lighting/StaticMesh.h
@@ -19,6 +19,9 @@ namespace Lightmass
 
 		FStaticLightingMapping* Mapping;
 
+		/** World to local transform, for box intersection w/ instances. */
+		FMatrix WorldToLocal;
+
 		virtual FStaticMeshStaticLightingMesh* GetInstanceableStaticMesh() override { return !bIsSplineMesh ? this : nullptr; }
 
 		virtual const FStaticMeshStaticLightingMesh* GetInstanceableStaticMesh() const override { return !bIsSplineMesh ? this : nullptr; }
@@ -47,7 +50,9 @@ namespace Lightmass
 		}
 
 		virtual void GetTriangle(int32 TriangleIndex,FStaticLightingVertex& OutV0,FStaticLightingVertex& OutV1,FStaticLightingVertex& OutV2,int32& ElementIndex) const;
+		virtual void GetNonTransformedTriangle(int32 TriangleIndex,FStaticLightingVertex& OutV0,FStaticLightingVertex& OutV1,FStaticLightingVertex& OutV2,int32& ElementIndex) const;
 		virtual void GetTriangleIndices(int32 TriangleIndex,int32& OutI0,int32& OutI1,int32& OutI2) const;
+		virtual void GetNonTransformedTriangleIndices(int32 TriangleIndex,int32& OutI0,int32& OutI1,int32& OutI2) const;
 
 		virtual bool IsElementCastingShadow(int32 ElementIndex) const;
 
diff --git a/Engine/Source/Programs/UnrealLightmass/Public/SceneExport.h b/Engine/Source/Programs/UnrealLightmass/Public/SceneExport.h
--- a/Engine/Source/Programs/UnrealLightmass/Public/SceneExport.h
+++ b/Engine/Source/Programs/UnrealLightmass/Public/SceneExport.h
@@ -72,9 +72,12 @@ public:
 	/** Whether to build Embree data structures for packet tracing. WIP feature - no lightmass algorithms emit packet tracing requests yet. */
 	bool bUseEmbreePacketTracing;
 
-	/** Whether to use kDOP trees & static mesh instancing to accelerate volumetric lightmap voxelization. Useful in scenes like large forest. */
+	/** Whether to use kDOP trees to accelerate volumetric lightmap voxelization. Useful in scenes like large forest. */
 	bool bUseFastVoxelization;
 
+	/** Whether to use static mesh instancing to reduce memory consumption in scenes like large forest. Might slow down small scenes. */
+	bool bUseEmbreeInstancing;
+
 	/** 
 	 * Direct lighting, skylight radiosity and irradiance photons are cached on mapping surfaces to accelerate final gathering.
 	 * This controls the downsample factor for that cache, relative to the mapping's lightmap resolution.
