Fix lightmass dependencies not being copied to output directory.

#rnx

diff --git a/Engine/Source/Programs/UnrealLightmass/UnrealLightmass.Build.cs b/Engine/Source/Programs/UnrealLightmass/UnrealLightmass.Build.cs
--- a/Engine/Source/Programs/UnrealLightmass/UnrealLightmass.Build.cs
+++ b/Engine/Source/Programs/UnrealLightmass/UnrealLightmass.Build.cs
@@ -66,9 +66,9 @@ public class UnrealLightmass : ModuleRules
             PublicIncludePaths.Add(SDKDir + "include");
             PublicAdditionalLibraries.Add(SDKDir + "lib/embree.lib");
 			RuntimeDependencies.Add("$(EngineDir)/Binaries/Win64/embree.dll");
-			RuntimeDependencies.Add("$(EngineDir)/Binaries/Win64/tbb.dll");
-			RuntimeDependencies.Add("$(EngineDir)/Binaries/Win64/tbbmalloc.dll");
-            PublicDefinitions.Add("USE_EMBREE=1");
+			RuntimeDependencies.Add("$(EngineDir)/Binaries/Win64/tbb.dll", Target.UEThirdPartySourceDirectory + "IntelEmbree/Embree2140/Win64/lib/tbb.dll"); // Take latest version to avoid overwriting the editor's copy
+			RuntimeDependencies.Add("$(EngineDir)/Binaries/Win64/tbbmalloc.dll", Target.UEThirdPartySourceDirectory + "IntelEmbree/Embree2140/Win64/lib/tbbmalloc.dll");
+			PublicDefinitions.Add("USE_EMBREE=1");
         }
 		else if (Target.Platform == UnrealTargetPlatform.Mac)
 		{
@@ -77,9 +77,9 @@ public class UnrealLightmass : ModuleRules
             PublicIncludePaths.Add(SDKDir + "include");
             PublicAdditionalLibraries.Add(SDKDir + "lib/libembree.2.dylib");
 			RuntimeDependencies.Add("$(EngineDir)/Binaries/Mac/libembree.2.dylib");
-			RuntimeDependencies.Add("$(EngineDir)/Binaries/Mac/libtbb.dylib");
-			RuntimeDependencies.Add("$(EngineDir)/Binaries/Mac/libtbbmalloc.dylib");
-            PublicDefinitions.Add("USE_EMBREE=1");
+			RuntimeDependencies.Add("$(EngineDir)/Binaries/Mac/libtbb.dylib", Target.UEThirdPartySourceDirectory + "IntelEmbree/Embree2140/MacOSX/lib/libtbb.dylib"); // Take latest version to avoid overwriting the editor's copy
+			RuntimeDependencies.Add("$(EngineDir)/Binaries/Mac/libtbbmalloc.dylib", Target.UEThirdPartySourceDirectory + "IntelEmbree/Embree2140/MacOSX/lib/libtbbmalloc.dylib");
+			PublicDefinitions.Add("USE_EMBREE=1");
 		}
         else
         {
