Improve dark corner mitigation by storing backfacing hits fraction inside irradiance cache records and marking texels as unmapped when their interpolated backfacing fraction > 0.5

diff --git a/Engine/Config/BaseLightmass.ini b/Engine/Config/BaseLightmass.ini
--- a/Engine/Config/BaseLightmass.ini
+++ b/Engine/Config/BaseLightmass.ini
@@ -40,7 +40,7 @@ MappingSurfaceCacheDownsampleFactor=2
 StaticLightingLevelScale=1
 VisibilityRayOffsetDistance=.1
 VisibilityNormalOffsetDistance=3
-VisibilityNormalOffsetSampleRadiusScale=.5
+VisibilityNormalOffsetSampleRadiusScale=.1
 VisibilityTangentOffsetSampleRadiusScale=.8
 SmallestTexelRadius=.1
 ; Tweaked for a good tradeoff between 'Cache Indirect Photon Paths' time and Indirect photon emitting 'Sampling Lights' time
diff --git a/Engine/Source/Programs/UnrealLightmass/Private/Lighting/FinalGather.cpp b/Engine/Source/Programs/UnrealLightmass/Private/Lighting/FinalGather.cpp
--- a/Engine/Source/Programs/UnrealLightmass/Private/Lighting/FinalGather.cpp
+++ b/Engine/Source/Programs/UnrealLightmass/Private/Lighting/FinalGather.cpp
@@ -1464,8 +1464,9 @@ FFinalGatherSample FStaticLightingSystem::CachePointIncomingRadiance(
 	const int32 BounceNumber = 1;
 	FFinalGatherSample IndirectLighting;
 	FFinalGatherSample UnusedSecondLighting;
+	float UnusedBackfacingHitsFraction;
 	// Attempt to interpolate incoming radiance from the lighting cache
-	if (!IrradianceCachingSettings.bAllowIrradianceCaching || !MappingContext.FirstBounceCache.InterpolateLighting(Vertex, true, bDebugThisTexel, 1.0f , IndirectLighting, UnusedSecondLighting, MappingContext.DebugCacheRecords))
+	if (!IrradianceCachingSettings.bAllowIrradianceCaching || !MappingContext.FirstBounceCache.InterpolateLighting(Vertex, true, bDebugThisTexel, 1.0f , IndirectLighting, UnusedSecondLighting, UnusedBackfacingHitsFraction, MappingContext.DebugCacheRecords))
 	{
 		// If final gathering is disabled, all indirect lighting will be estimated using photon mapping.
 		// This is really only useful for debugging since it requires an excessive number of indirect photons to get indirect shadows for the first bounce.
diff --git a/Engine/Source/Programs/UnrealLightmass/Private/Lighting/LightingCache.h b/Engine/Source/Programs/UnrealLightmass/Private/Lighting/LightingCache.h
--- a/Engine/Source/Programs/UnrealLightmass/Private/Lighting/LightingCache.h
+++ b/Engine/Source/Programs/UnrealLightmass/Private/Lighting/LightingCache.h
@@ -180,6 +180,8 @@ public:
 		/** For debugging */
 		int32 Id;
 
+		float BackfacingHitsFraction;
+
 		/** Initialization constructor. */
 		FRecord(const FFullStaticLightingVertex& InVertex,int32 InElementIndex,const FLightingCacheGatherInfo& GatherInfo,float SampleRadius,float InOverrideRadius,const FIrradianceCachingSettings& IrradianceCachingSettings,const FStaticLightingSettings& GeneralSettings,const RecordSampleType& InLighting, const FVector4& InRotGradient, const FVector4& InTransGradient):
 			Vertex(InVertex),
@@ -195,6 +197,8 @@ public:
 			InterpolationRadius = Radius * FMath::Max(IrradianceCachingSettings.DistanceSmoothFactor * GeneralSettings.IndirectLightingSmoothness, 1.0f);
 
 			BoundingRadius = FMath::Max(Radius, InterpolationRadius);
+
+			BackfacingHitsFraction = GatherInfo.BackfacingHitsFraction;
 		}
 	};
 
@@ -282,6 +286,7 @@ public:
 		float SecondInterpolationSmoothnessReduction,
 		SampleType& OutLighting,
 		SampleType& OutSecondLighting,
+		float& OutBackfacingHitsFraction,
 		TArray<FDebugLightingCacheRecord>& DebugCacheRecords,
 		class FInfluencingRecordCollector* RecordCollector = NULL) const;
 
@@ -359,6 +364,7 @@ bool TLightingCache<SampleType>::InterpolateLighting(
 	float SecondInterpolationSmoothnessReduction,
 	SampleType& OutLighting,
 	SampleType& OutSecondLighting,
+	float& OutBackfacingHitsFraction,
 	TArray<FDebugLightingCacheRecord>& DebugCacheRecords,
 	FInfluencingRecordCollector* RecordCollector) const
 {
@@ -372,6 +378,7 @@ bool TLightingCache<SampleType>::InterpolateLighting(
 	float TotalWeight = 0.0f;
 	SampleType SecondAccumulatedLighting(ForceInit);
 	float SecondTotalWeight = 0.0f;
+	float AccumulatedBackfacingHitsFraction = 0.0f;
 
 	// Iterate over the octree nodes containing the query point.
 	for( typename LightingOctreeType::template TConstElementBoxIterator<> OctreeIt(
@@ -433,6 +440,7 @@ bool TLightingCache<SampleType>::InterpolateLighting(
 
 				//@todo - Rotate the record's lighting into this vertex's tangent basis.  We are linearly combining incident lighting in different coordinate spaces.
 				AccumulatedLighting = AccumulatedLighting + LightingRecord.Lighting * RecordWeight * (NonGradientLighting + RotationalGradientContribution + TranslationalGradientContribution);
+				AccumulatedBackfacingHitsFraction += LightingRecord.BackfacingHitsFraction * RecordWeight * (NonGradientLighting + RotationalGradientContribution + TranslationalGradientContribution);
 				// Accumulate the weight of all records
 				TotalWeight += RecordWeight;
 
@@ -498,6 +506,7 @@ bool TLightingCache<SampleType>::InterpolateLighting(
 		const float InvTotalWeight = 1.0f / TotalWeight;
 		OutLighting = OutLighting + AccumulatedLighting * InvTotalWeight;
 		OutSecondLighting = OutSecondLighting + SecondAccumulatedLighting * (1.0f / SecondTotalWeight);
+		OutBackfacingHitsFraction = AccumulatedBackfacingHitsFraction * InvTotalWeight;
 		return true;
 	}
 	else
diff --git a/Engine/Source/Programs/UnrealLightmass/Private/Lighting/Radiosity.cpp b/Engine/Source/Programs/UnrealLightmass/Private/Lighting/Radiosity.cpp
--- a/Engine/Source/Programs/UnrealLightmass/Private/Lighting/Radiosity.cpp
+++ b/Engine/Source/Programs/UnrealLightmass/Private/Lighting/Radiosity.cpp
@@ -298,8 +298,9 @@ void FStaticLightingSystem::RadiositySetupTextureMapping(FStaticLightingTextureM
 					Vertex.ApplyVertexModifications(TexelToVertex.ElementIndex, MaterialSettings.bUseNormalMapsForLighting, TextureMapping->Mesh);
 					FFinalGatherSample SkyLighting;
 					FFinalGatherSample UnusedSecondLighting;
+					float UnusedBackfacingHitsFraction;
 
-					if (!RadiosityCache.InterpolateLighting(Vertex, true, false, 1.0f, SkyLighting, UnusedSecondLighting, MappingContext.DebugCacheRecords))
+					if (!RadiosityCache.InterpolateLighting(Vertex, true, false, 1.0f, SkyLighting, UnusedSecondLighting, UnusedBackfacingHitsFraction, MappingContext.DebugCacheRecords))
 					{
 						FFinalGatherSample UniformSampledIncomingRadiance;
 						TArray<FVector4> ImportancePhotonDirections;
@@ -403,7 +404,8 @@ void FStaticLightingSystem::RadiositySetupTextureMapping(FStaticLightingTextureM
 				{
 					FFinalGatherSample SkyLighting;
 					FFinalGatherSample UnusedSecondLighting;
-					RadiosityCache.InterpolateLighting(CurrentVertex, false, false, IrradianceCachingSettings.SkyOcclusionSmoothnessReduction, SkyLighting, UnusedSecondLighting, MappingContext.DebugCacheRecords, RecordCollectorPtr);
+					float UnusedBackfacingHitsFraction;
+					RadiosityCache.InterpolateLighting(CurrentVertex, false, false, IrradianceCachingSettings.SkyOcclusionSmoothnessReduction, SkyLighting, UnusedSecondLighting, UnusedBackfacingHitsFraction, MappingContext.DebugCacheRecords, RecordCollectorPtr);
 
 					if (GeneralSettings.ViewSingleBounceNumber < 0 || GeneralSettings.ViewSingleBounceNumber == 1)
 					{
@@ -611,8 +613,9 @@ void FStaticLightingSystem::RadiosityIterationTextureMapping(FStaticLightingText
 					Vertex.ApplyVertexModifications(TexelToVertex.ElementIndex, MaterialSettings.bUseNormalMapsForLighting, TextureMapping->Mesh);
 					FFinalGatherSample SkyLighting;
 					FFinalGatherSample UnusedSecondLighting;
+					float UnusedBackfacingHitsFraction;
 
-					if (!RadiosityCache.InterpolateLighting(Vertex, true, false, 1.0f, SkyLighting, UnusedSecondLighting, MappingContext.DebugCacheRecords))
+					if (!RadiosityCache.InterpolateLighting(Vertex, true, false, 1.0f, SkyLighting, UnusedSecondLighting, UnusedBackfacingHitsFraction, MappingContext.DebugCacheRecords))
 					{
 						FFinalGatherSample UniformSampledIncomingRadiance;
 						//@todo - find and pass in photons from the appropriate bounce number to improve bUseRadiositySolverForLightMultibounce quality
@@ -689,7 +692,8 @@ void FStaticLightingSystem::RadiosityIterationTextureMapping(FStaticLightingText
 
 					FFinalGatherSample IterationLighting;
 					FFinalGatherSample UnusedSecondLighting;
-					RadiosityCache.InterpolateLighting(CurrentVertex, false, false, IrradianceCachingSettings.SkyOcclusionSmoothnessReduction, IterationLighting, UnusedSecondLighting, MappingContext.DebugCacheRecords);
+					float UnusedBackfacingHitsFraction;
+					RadiosityCache.InterpolateLighting(CurrentVertex, false, false, IrradianceCachingSettings.SkyOcclusionSmoothnessReduction, IterationLighting, UnusedSecondLighting, UnusedBackfacingHitsFraction, MappingContext.DebugCacheRecords);
 				
 					const bool bIsTranslucent = TextureMapping->Mesh->IsTranslucent(TexelToVertex.ElementIndex);
 					const FLinearColor Reflectance = (bIsTranslucent ? FLinearColor::Black : TextureMapping->Mesh->EvaluateTotalReflectance(CurrentVertex, TexelToVertex.ElementIndex)) * (float)INV_PI;
diff --git a/Engine/Source/Programs/UnrealLightmass/Private/Lighting/TextureMapping.cpp b/Engine/Source/Programs/UnrealLightmass/Private/Lighting/TextureMapping.cpp
--- a/Engine/Source/Programs/UnrealLightmass/Private/Lighting/TextureMapping.cpp
+++ b/Engine/Source/Programs/UnrealLightmass/Private/Lighting/TextureMapping.cpp
@@ -2984,14 +2984,16 @@ void FStaticLightingSystem::ProcessInterpolateTask(FInterpolateIndirectTaskDescr
 				FFullStaticLightingVertex TexelVertex = TexelToVertex.GetFullVertex();
 				FFinalGatherSample IndirectLighting;
 				FFinalGatherSample SecondInterpolatedIndirectLighting;
+				float BackfacingHitsFraction = 0.0f;
+				float BackFaceBackfacingHitsFraction = 1.0f;
 				// Interpolate the indirect lighting from the irradiance cache
 				// Interpolation must succeed since this is the second pass
-				verify(Task->FirstBounceCache->InterpolateLighting(TexelVertex, false, bDebugThisTexel && GeneralSettings.ViewSingleBounceNumber == 1, IrradianceCachingSettings.SkyOcclusionSmoothnessReduction, IndirectLighting, SecondInterpolatedIndirectLighting, Task->MappingContext.DebugCacheRecords));
+				verify(Task->FirstBounceCache->InterpolateLighting(TexelVertex, false, bDebugThisTexel && GeneralSettings.ViewSingleBounceNumber == 1, IrradianceCachingSettings.SkyOcclusionSmoothnessReduction, IndirectLighting, SecondInterpolatedIndirectLighting, BackfacingHitsFraction, Task->MappingContext.DebugCacheRecords));
 
 				// Replace sky occlusion in the lighting sample that will be written into the lightmap with the interpolated sky occlusion using IrradianceCachingSettings.SkyOcclusionSmoothnessReduction
 				IndirectLighting.SkyOcclusion = SecondInterpolatedIndirectLighting.SkyOcclusion;
 				IndirectLighting.StationarySkyLighting = SecondInterpolatedIndirectLighting.StationarySkyLighting;
-
+				
 				if (Task->TextureMapping->Mesh->UsesTwoSidedLighting(TexelToVertex.ElementIndex))
 				{
 					TexelVertex.WorldTangentX = -TexelVertex.WorldTangentX;
@@ -3001,12 +3003,17 @@ void FStaticLightingSystem::ProcessInterpolateTask(FInterpolateIndirectTaskDescr
 					FFinalGatherSample BackFaceIndirectLighting;
 					FFinalGatherSample BackFaceSecondInterpolatedIndirectLighting;
 					// Interpolate indirect lighting for the back face
-					verify(Task->FirstBounceCache->InterpolateLighting(TexelVertex, false, bDebugThisTexel && GeneralSettings.ViewSingleBounceNumber == 1, IrradianceCachingSettings.SkyOcclusionSmoothnessReduction, BackFaceIndirectLighting, BackFaceSecondInterpolatedIndirectLighting, Task->MappingContext.DebugCacheRecords));
+					verify(Task->FirstBounceCache->InterpolateLighting(TexelVertex, false, bDebugThisTexel && GeneralSettings.ViewSingleBounceNumber == 1, IrradianceCachingSettings.SkyOcclusionSmoothnessReduction, BackFaceIndirectLighting, BackFaceSecondInterpolatedIndirectLighting, BackFaceBackfacingHitsFraction, Task->MappingContext.DebugCacheRecords));
 					BackFaceIndirectLighting.SkyOcclusion = BackFaceSecondInterpolatedIndirectLighting.SkyOcclusion;
 					// Average front and back face incident lighting
 					IndirectLighting = (BackFaceIndirectLighting + IndirectLighting) * 0.5f;
 				}
 
+				if (BackfacingHitsFraction > 0.5f && BackFaceBackfacingHitsFraction > 0.5f)
+				{
+					CurrentLightSample.bIsMapped = false;
+				}
+
 				float IndirectOcclusion = 1.0f;
 				if (AmbientOcclusionSettings.bUseAmbientOcclusion)
 				{
