Account for texel size to avoid creating tiny irradiance cache entries and post-aliasing issues

diff --git a/Engine/Source/Programs/UnrealLightmass/Private/Lighting/FinalGather.cpp b/Engine/Source/Programs/UnrealLightmass/Private/Lighting/FinalGather.cpp
--- a/Engine/Source/Programs/UnrealLightmass/Private/Lighting/FinalGather.cpp
+++ b/Engine/Source/Programs/UnrealLightmass/Private/Lighting/FinalGather.cpp
@@ -1447,6 +1447,7 @@ FFinalGatherSample FStaticLightingSystem::CachePointIncomingRadiance(
 	const FStaticLightingMapping* Mapping,
 	const FFullStaticLightingVertex& Vertex,
 	int32 ElementIndex,
+	float TexelRadius,
 	float SampleRadius,
 	bool bIntersectingSurface,
 	FStaticLightingMappingContext& MappingContext,
@@ -1692,7 +1693,7 @@ FFinalGatherSample FStaticLightingSystem::CachePointIncomingRadiance(
 					Vertex,
 					ElementIndex,
 					GatherInfo,
-					SampleRadius,
+					BounceNumber == 1 ? TexelRadius : SampleRadius,
 					OverrideRadius,
 					IrradianceCachingSettings,
 					GeneralSettings,
diff --git a/Engine/Source/Programs/UnrealLightmass/Private/Lighting/LightingCache.h b/Engine/Source/Programs/UnrealLightmass/Private/Lighting/LightingCache.h
--- a/Engine/Source/Programs/UnrealLightmass/Private/Lighting/LightingCache.h
+++ b/Engine/Source/Programs/UnrealLightmass/Private/Lighting/LightingCache.h
@@ -192,7 +192,7 @@ public:
 			Id(-1)
 		{
 			// Clamp to be larger than the texel
-			Radius = FMath::Clamp(GatherInfo.MinDistance, SampleRadius, IrradianceCachingSettings.MaxRecordRadius) * IrradianceCachingSettings.RecordRadiusScale; 
+			Radius = FMath::Clamp(GatherInfo.MinDistance * IrradianceCachingSettings.RecordRadiusScale, SampleRadius, IrradianceCachingSettings.MaxRecordRadius);
 			// Use a larger radius to interpolate, which smooths the error
 			InterpolationRadius = Radius * FMath::Max(IrradianceCachingSettings.DistanceSmoothFactor * GeneralSettings.IndirectLightingSmoothness, 1.0f);
 
diff --git a/Engine/Source/Programs/UnrealLightmass/Private/Lighting/LightingSystem.h b/Engine/Source/Programs/UnrealLightmass/Private/Lighting/LightingSystem.h
--- a/Engine/Source/Programs/UnrealLightmass/Private/Lighting/LightingSystem.h
+++ b/Engine/Source/Programs/UnrealLightmass/Private/Lighting/LightingSystem.h
@@ -2384,6 +2384,7 @@ private:
 		const FStaticLightingMapping* Mapping,
 		const FFullStaticLightingVertex& Vertex,
 		int32 ElementIndex,
+		float TexelRadius,
 		float SampleRadius,
 		bool bIntersectingSurface,
 		FStaticLightingMappingContext& MappingContext,
diff --git a/Engine/Source/Programs/UnrealLightmass/Private/Lighting/Radiosity.cpp b/Engine/Source/Programs/UnrealLightmass/Private/Lighting/Radiosity.cpp
--- a/Engine/Source/Programs/UnrealLightmass/Private/Lighting/Radiosity.cpp
+++ b/Engine/Source/Programs/UnrealLightmass/Private/Lighting/Radiosity.cpp
@@ -316,7 +316,7 @@ void FStaticLightingSystem::RadiositySetupTextureMapping(FStaticLightingTextureM
 						UniformSampledIncomingRadiance = IncomingRadianceAdaptive<FFinalGatherSample>(
 							TextureMapping, 
 							Vertex, 
-							TexelToVertex.TexelRadius, 
+							TexelToVertex.SampleRadius, 
 							false,
 							TexelToVertex.ElementIndex, 
 							2,  /** BounceNumber */
@@ -340,7 +340,7 @@ void FStaticLightingSystem::RadiositySetupTextureMapping(FStaticLightingTextureM
 							Vertex,
 							TexelToVertex.ElementIndex,
 							GatherInfo,
-							TexelToVertex.TexelRadius,
+							TexelToVertex.SampleRadius,
 							OverrideRadius,
 							IrradianceCachingSettings,
 							GeneralSettings,
@@ -628,7 +628,7 @@ void FStaticLightingSystem::RadiosityIterationTextureMapping(FStaticLightingText
 						UniformSampledIncomingRadiance = IncomingRadianceAdaptive<FFinalGatherSample>(
 							TextureMapping, 
 							Vertex, 
-							TexelToVertex.TexelRadius, 
+							TexelToVertex.SampleRadius, 
 							false,
 							TexelToVertex.ElementIndex, 
 							PassIndex + 3, /** BounceNumber */
@@ -652,7 +652,7 @@ void FStaticLightingSystem::RadiosityIterationTextureMapping(FStaticLightingText
 							Vertex,
 							TexelToVertex.ElementIndex,
 							GatherInfo,
-							TexelToVertex.TexelRadius,
+							TexelToVertex.SampleRadius,
 							OverrideRadius,
 							IrradianceCachingSettings,
 							GeneralSettings,
diff --git a/Engine/Source/Programs/UnrealLightmass/Private/Lighting/TextureMapping.cpp b/Engine/Source/Programs/UnrealLightmass/Private/Lighting/TextureMapping.cpp
--- a/Engine/Source/Programs/UnrealLightmass/Private/Lighting/TextureMapping.cpp
+++ b/Engine/Source/Programs/UnrealLightmass/Private/Lighting/TextureMapping.cpp
@@ -486,6 +486,7 @@ void FStaticLightingSystem::AdjustRepresentativeSurfelForTexelsTextureMapping(
 					MinDistanceSquared = SceneConstants.SmallestTexelRadius;
 				}
 				TexelToVertex.TexelRadius = FMath::Max(FMath::Sqrt(MinDistanceSquared), SceneConstants.SmallestTexelRadius);
+				TexelToVertex.SampleRadius = TexelToVertex.TexelRadius;
 				MappingContext.Stats.NumMappedTexels++;
 
 				{
@@ -554,7 +555,7 @@ void FStaticLightingSystem::AdjustRepresentativeSurfelForTexelsTextureMapping(
 					{
 						checkSlow(Intersections[ClosestIntersectionIndex].bIntersects);
 
-						TexelToVertex.TexelRadius = FMath::Min(TexelToVertex.TexelRadius, FMath::Sqrt(ClosestIntersectionDistanceSq / 2.0f));
+						TexelToVertex.SampleRadius = FMath::Min(TexelToVertex.TexelRadius, FMath::Sqrt(ClosestIntersectionDistanceSq / 2.0f));
 					}
 
 					// Give preference to moving the shading position outside of backfaces
@@ -579,10 +580,10 @@ void FStaticLightingSystem::AdjustRepresentativeSurfelForTexelsTextureMapping(
 						// Project back onto plane of texel to avoid incorrect self occlusion
 						TexelToVertex.WorldPosition = OffsetShadingPosition + TexelToVertex.TriangleNormal * Dot3(TexelToVertex.TriangleNormal, TexelToVertex.WorldPosition - OffsetShadingPosition);
 
-						TexelToVertex.TexelRadius = (OffsetShadingPosition - Intersections[IntersectionIndexForShadingPositionMovement].IntersectionVertex.WorldPosition).Size3() / FMath::Sqrt(2.0f);
+						TexelToVertex.SampleRadius = (OffsetShadingPosition - Intersections[IntersectionIndexForShadingPositionMovement].IntersectionVertex.WorldPosition).Size3() / FMath::Sqrt(2.0f);
 					}
 
-					TexelToVertex.TexelRadius = FMath::Max(TexelToVertex.TexelRadius, SceneConstants.SmallestTexelRadius);
+					TexelToVertex.SampleRadius = FMath::Max(TexelToVertex.SampleRadius, SceneConstants.SmallestTexelRadius);
 				}
 			}
 			else
@@ -2907,6 +2908,7 @@ void FStaticLightingSystem::ProcessCacheIndirectLightingTask(FCacheIndirectTaskD
 					TexelVertex, 
 					TexelToVertex.ElementIndex,
 					TexelToVertex.TexelRadius, 
+					TexelToVertex.SampleRadius, 
 					TexelToVertex.bIntersectingSurface,
 					Task->MappingContext, 
 					SampleGenerator, 
@@ -2924,6 +2926,7 @@ void FStaticLightingSystem::ProcessCacheIndirectLightingTask(FCacheIndirectTaskD
 						TexelVertex, 
 						TexelToVertex.ElementIndex,
 						TexelToVertex.TexelRadius, 
+						TexelToVertex.SampleRadius, 
 						TexelToVertex.bIntersectingSurface,
 						Task->MappingContext, 
 						SampleGenerator, 
diff --git a/Engine/Source/Programs/UnrealLightmass/Private/Lighting/TextureMappingSetup.h b/Engine/Source/Programs/UnrealLightmass/Private/Lighting/TextureMappingSetup.h
--- a/Engine/Source/Programs/UnrealLightmass/Private/Lighting/TextureMappingSetup.h
+++ b/Engine/Source/Programs/UnrealLightmass/Private/Lighting/TextureMappingSetup.h
@@ -26,6 +26,9 @@ struct FTexelToVertex
 	/** World space radius of the texel. */
 	float TexelRadius;
 
+	/** Final gathering sample radius, derived from texel radius by adjustment */
+	float SampleRadius;
+
 	/** Whether this texel was determined to be intersecting another surface. */
 	uint32 bIntersectingSurface : 1;
 
