Lightmass: revert the lighting cache sample radius change as it is causing entries in thin walls to leak outside due to their radius being larger than the thickness of the wall

diff --git a/Engine/Source/Programs/UnrealLightmass/Private/Lighting/LightingCache.h b/Engine/Source/Programs/UnrealLightmass/Private/Lighting/LightingCache.h
--- a/Engine/Source/Programs/UnrealLightmass/Private/Lighting/LightingCache.h
+++ b/Engine/Source/Programs/UnrealLightmass/Private/Lighting/LightingCache.h
@@ -192,7 +192,7 @@ public:
 			Id(-1)
 		{
 			// Clamp to be larger than the texel
-			Radius = FMath::Clamp(GatherInfo.MinDistance * IrradianceCachingSettings.RecordRadiusScale, SampleRadius, IrradianceCachingSettings.MaxRecordRadius);
+			Radius = FMath::Clamp(GatherInfo.MinDistance, SampleRadius, IrradianceCachingSettings.MaxRecordRadius) * IrradianceCachingSettings.RecordRadiusScale;
 			// Use a larger radius to interpolate, which smooths the error
 			InterpolationRadius = Radius * FMath::Max(IrradianceCachingSettings.DistanceSmoothFactor * GeneralSettings.IndirectLightingSmoothness, 1.0f);
 
