Fix Lightmass backface detection to also consider two sided materials

diff --git a/Engine/Source/Programs/UnrealLightmass/Private/Lighting/TextureMapping.cpp b/Engine/Source/Programs/UnrealLightmass/Private/Lighting/TextureMapping.cpp
--- a/Engine/Source/Programs/UnrealLightmass/Private/Lighting/TextureMapping.cpp
+++ b/Engine/Source/Programs/UnrealLightmass/Private/Lighting/TextureMapping.cpp
@@ -1074,7 +1074,7 @@ void FStaticLightingSystem::TraceToTexelCorner(
 
 	AggregateMesh->IntersectLightRay(TexelRay, true, false, false, MappingContext.RayCache, Intersection);
 
-	bHitBackface = Intersection.bIntersects && Dot3(Intersection.IntersectionVertex.WorldTangentZ, TexelRay.Direction) >= 0;
+	bHitBackface = Intersection.bIntersects && Dot3(Intersection.IntersectionVertex.WorldTangentZ, TexelRay.Direction) >= 0 && !Intersection.Mesh->IsTwoSided(Intersection.ElementIndex);
 
 #if ALLOW_LIGHTMAP_SAMPLE_DEBUGGING
 	if (bDebugThisTexel)
