M - Rename the Octree to Octree2 in preparation to add a deprecation for the original class

diff --git a/Engine/Plugins/MeshPainting/Source/MeshPaintingToolset/Private/MeshPaintHelpers.cpp b/Engine/Plugins/MeshPainting/Source/MeshPaintingToolset/Private/MeshPaintHelpers.cpp
--- a/Engine/Plugins/MeshPainting/Source/MeshPaintingToolset/Private/MeshPaintHelpers.cpp
+++ b/Engine/Plugins/MeshPainting/Source/MeshPaintingToolset/Private/MeshPaintHelpers.cpp
@@ -1248,7 +1248,7 @@ struct FVertexColorPropogationOctreeSemantics
 	{
 	}
 };
-typedef TOctree<FPaintedMeshVertex, FVertexColorPropogationOctreeSemantics> TVertexColorPropogationPosOctree;
+typedef TOctree2<FPaintedMeshVertex, FVertexColorPropogationOctreeSemantics> TVertexColorPropogationPosOctree;
 
 void UMeshPaintingToolset::ApplyVertexColorsToAllLODs(IMeshPaintComponentAdapter& GeometryInfo, USkeletalMeshComponent* SkeletalMeshComponent)
 {
diff --git a/Engine/Plugins/MeshPainting/Source/MeshPaintingToolset/Public/TMeshPaintOctree.h b/Engine/Plugins/MeshPainting/Source/MeshPaintingToolset/Public/TMeshPaintOctree.h
--- a/Engine/Plugins/MeshPainting/Source/MeshPaintingToolset/Public/TMeshPaintOctree.h
+++ b/Engine/Plugins/MeshPainting/Source/MeshPaintingToolset/Public/TMeshPaintOctree.h
@@ -56,4 +56,4 @@ struct FMeshPaintTriangleOctreeSemantics
 
 	}
 };
-typedef TOctree<FMeshPaintTriangle, FMeshPaintTriangleOctreeSemantics> FMeshPaintTriangleOctree;
\ No newline at end of file
+typedef TOctree2<FMeshPaintTriangle, FMeshPaintTriangleOctreeSemantics> FMeshPaintTriangleOctree;
\ No newline at end of file
diff --git a/Engine/Plugins/Runtime/EditableMesh/Source/EditableMesh/EditableMeshOctree.h b/Engine/Plugins/Runtime/EditableMesh/Source/EditableMesh/EditableMeshOctree.h
--- a/Engine/Plugins/Runtime/EditableMesh/Source/EditableMesh/EditableMeshOctree.h
+++ b/Engine/Plugins/Runtime/EditableMesh/Source/EditableMesh/EditableMeshOctree.h
@@ -54,11 +54,11 @@ struct FEditableMeshOctreeSemantics
 	}
 };
 
-class FEditableMeshOctree : public TOctree<FEditableMeshOctreePolygon, FEditableMeshOctreeSemantics>
+class FEditableMeshOctree : public TOctree2<FEditableMeshOctreePolygon, FEditableMeshOctreeSemantics>
 {
 public:
 	FEditableMeshOctree( const FVector& InOrigin, float InExtent )
-		: TOctree( InOrigin, InExtent )
+		: TOctree2( InOrigin, InExtent )
 	{
 	}
 };
diff --git a/Engine/Source/Developer/SkeletalMeshUtilitiesCommon/Private/LODUtilities.cpp b/Engine/Source/Developer/SkeletalMeshUtilitiesCommon/Private/LODUtilities.cpp
--- a/Engine/Source/Developer/SkeletalMeshUtilitiesCommon/Private/LODUtilities.cpp
+++ b/Engine/Source/Developer/SkeletalMeshUtilitiesCommon/Private/LODUtilities.cpp
@@ -1346,7 +1346,7 @@ struct FTriangleOctreeSemantics
 	}
 };
 
-typedef TOctree<FTriangleElement, FTriangleOctreeSemantics> TTriangleElementOctree;
+typedef TOctree2<FTriangleElement, FTriangleOctreeSemantics> TTriangleElementOctree;
 
 void MatchVertexIndexUsingPosition(
 	const FSkeletalMeshImportData& ImportDataDest
diff --git a/Engine/Source/Editor/MeshPaint/Private/MeshPaintHelpers.cpp b/Engine/Source/Editor/MeshPaint/Private/MeshPaintHelpers.cpp
--- a/Engine/Source/Editor/MeshPaint/Private/MeshPaintHelpers.cpp
+++ b/Engine/Source/Editor/MeshPaint/Private/MeshPaintHelpers.cpp
@@ -1759,7 +1759,7 @@ struct FVertexColorPropogationOctreeSemantics
 	{
 	}
 };
-typedef TOctree<FPaintedMeshVertex, FVertexColorPropogationOctreeSemantics> TVertexColorPropogationPosOctree;
+typedef TOctree2<FPaintedMeshVertex, FVertexColorPropogationOctreeSemantics> TVertexColorPropogationPosOctree;
 
 void MeshPaintHelpers::ApplyVertexColorsToAllLODs(IMeshPaintGeometryAdapter& GeometryInfo, USkeletalMeshComponent* SkeletalMeshComponent)
 {
diff --git a/Engine/Source/Editor/MeshPaint/Public/TMeshPaintOctree.h b/Engine/Source/Editor/MeshPaint/Public/TMeshPaintOctree.h
--- a/Engine/Source/Editor/MeshPaint/Public/TMeshPaintOctree.h
+++ b/Engine/Source/Editor/MeshPaint/Public/TMeshPaintOctree.h
@@ -56,4 +56,4 @@ struct FMeshPaintTriangleOctreeSemantics
 
 	}
 };
-typedef TOctree<FMeshPaintTriangle, FMeshPaintTriangleOctreeSemantics> FMeshPaintTriangleOctree;
\ No newline at end of file
+typedef TOctree2<FMeshPaintTriangle, FMeshPaintTriangleOctreeSemantics> FMeshPaintTriangleOctree;
\ No newline at end of file
diff --git a/Engine/Source/Editor/UnrealEd/Private/Fbx/FbxSkeletalMeshImport.cpp b/Engine/Source/Editor/UnrealEd/Private/Fbx/FbxSkeletalMeshImport.cpp
--- a/Engine/Source/Editor/UnrealEd/Private/Fbx/FbxSkeletalMeshImport.cpp
+++ b/Engine/Source/Editor/UnrealEd/Private/Fbx/FbxSkeletalMeshImport.cpp
@@ -160,7 +160,7 @@ struct FSkeletalMeshVertPosOctreeSemantics
 	{
 	}
 };
-typedef TOctree<FSoftSkinVertex, FSkeletalMeshVertPosOctreeSemantics> TSKCVertPosOctree;
+typedef TOctree2<FSoftSkinVertex, FSkeletalMeshVertPosOctreeSemantics> TSKCVertPosOctree;
 
 void RemapSkeletalMeshVertexColorToImportData(const USkeletalMesh* SkeletalMesh, const int32 LODIndex, FSkeletalMeshImportData* SkelMeshImportData)
 {
diff --git a/Engine/Source/Programs/UnrealLightmass/Private/Launch/UnitTest.cpp b/Engine/Source/Programs/UnrealLightmass/Private/Launch/UnitTest.cpp
--- a/Engine/Source/Programs/UnrealLightmass/Private/Launch/UnitTest.cpp
+++ b/Engine/Source/Programs/UnrealLightmass/Private/Launch/UnitTest.cpp
@@ -217,7 +217,7 @@ void TestLightmass()
 
 	UE_LOG(LogLightmass, Display, TEXT("sizeof FDirectionalLight = %d, FLight = %d, FDirectionalLightData = %d"), sizeof(FDirectionalLight), sizeof(FLight), sizeof(FDirectionalLightData));
 
-	TOctree<float, FTestOctreeSemantics> TestOctree(FVector4(0), 10.0f);
+	TOctree2<float, FTestOctreeSemantics> TestOctree(FVector4(0), 10.0f);
 	TestOctree.AddElement(5);
 
 
diff --git a/Engine/Source/Programs/UnrealLightmass/Private/Lighting/LightingCache.h b/Engine/Source/Programs/UnrealLightmass/Private/Lighting/LightingCache.h
--- a/Engine/Source/Programs/UnrealLightmass/Private/Lighting/LightingCache.h
+++ b/Engine/Source/Programs/UnrealLightmass/Private/Lighting/LightingCache.h
@@ -205,7 +205,7 @@ public:
 	struct FRecordOctreeSemantics;
 
 	/** The type of lighting cache octree nodes. */
-	typedef TOctree<FRecord<SampleType>,FRecordOctreeSemantics> LightingOctreeType;
+	typedef TOctree2<FRecord<SampleType>,FRecordOctreeSemantics> LightingOctreeType;
 	
 	/** The octree semantics for irradiance records. */
 	struct FRecordOctreeSemantics
diff --git a/Engine/Source/Programs/UnrealLightmass/Private/Lighting/LightingSystem.h b/Engine/Source/Programs/UnrealLightmass/Private/Lighting/LightingSystem.h
--- a/Engine/Source/Programs/UnrealLightmass/Private/Lighting/LightingSystem.h
+++ b/Engine/Source/Programs/UnrealLightmass/Private/Lighting/LightingSystem.h
@@ -231,7 +231,7 @@ struct FPhotonElement
 	{}
 };
 
-typedef TOctree<FPhotonElement,struct FPhotonMapOctreeSemantics> FPhotonOctree;
+typedef TOctree2<FPhotonElement,struct FPhotonMapOctreeSemantics> FPhotonOctree;
 
 /** Octree semantic definitions. */
 struct FPhotonMapOctreeSemantics
@@ -299,7 +299,7 @@ struct FPhotonSegmentMapOctreeSemantics
 	}
 };
 
-typedef TOctree<FPhotonSegmentElement,struct FPhotonSegmentMapOctreeSemantics> FPhotonSegmentOctree;
+typedef TOctree2<FPhotonSegmentElement,struct FPhotonSegmentMapOctreeSemantics> FPhotonSegmentOctree;
 
 /** A photon which stores a precalculated irradiance estimate. */
 class FIrradiancePhoton : public FIrradiancePhotonData
@@ -376,7 +376,7 @@ private:
 	TArray<FIrradiancePhoton>& PhotonArray;
 };
 
-typedef TOctree<FIrradiancePhotonElement,struct FIrradiancePhotonMapOctreeSemantics> FIrradiancePhotonOctree;
+typedef TOctree2<FIrradiancePhotonElement,struct FIrradiancePhotonMapOctreeSemantics> FIrradiancePhotonOctree;
 
 /** Octree semantic definitions. */
 struct FIrradiancePhotonMapOctreeSemantics
@@ -438,7 +438,7 @@ struct FVolumeSampleInterpolationElement
 	{}
 };
 
-typedef TOctree<FVolumeSampleInterpolationElement,struct FVolumeLightingInterpolationOctreeSemantics> FVolumeLightingInterpolationOctree;
+typedef TOctree2<FVolumeSampleInterpolationElement,struct FVolumeLightingInterpolationOctreeSemantics> FVolumeLightingInterpolationOctree;
 
 /** Octree semantic definitions. */
 struct FVolumeLightingInterpolationOctreeSemantics
diff --git a/Engine/Source/Programs/UnrealLightmass/Private/Lighting/SampleVolume.cpp b/Engine/Source/Programs/UnrealLightmass/Private/Lighting/SampleVolume.cpp
--- a/Engine/Source/Programs/UnrealLightmass/Private/Lighting/SampleVolume.cpp
+++ b/Engine/Source/Programs/UnrealLightmass/Private/Lighting/SampleVolume.cpp
@@ -12,7 +12,7 @@ namespace Lightmass
 
 typedef FVolumeSampleInterpolationElement FVolumeSampleProximityElement;
 
-typedef TOctree<FVolumeSampleProximityElement,struct FVolumeLightingProximityOctreeSemantics> FVolumeLightingProximityOctree;
+typedef TOctree2<FVolumeSampleProximityElement,struct FVolumeLightingProximityOctreeSemantics> FVolumeLightingProximityOctree;
 
 struct FVolumeLightingProximityOctreeSemantics
 {
diff --git a/Engine/Source/Programs/UnrealLightmass/Private/LightmassCore/Math/LMOctree.h b/Engine/Source/Programs/UnrealLightmass/Private/LightmassCore/Math/LMOctree.h
--- a/Engine/Source/Programs/UnrealLightmass/Private/LightmassCore/Math/LMOctree.h
+++ b/Engine/Source/Programs/UnrealLightmass/Private/LightmassCore/Math/LMOctree.h
@@ -702,13 +702,13 @@ FORCEINLINE FOctreeChildNodeRef FOctreeNodeContext::GetContainingChild(const FBo
 }
 
 template<typename ElementType,typename OctreeSemantics>
-void TOctree<ElementType,OctreeSemantics>::AddElement(typename TTypeTraits<ElementType>::ConstInitType Element)
+void TOctree2<ElementType,OctreeSemantics>::AddElement(typename TTypeTraits<ElementType>::ConstInitType Element)
 {
 	AddElementToNode(Element,RootNode,RootNodeContext);
 }
 															
 template<typename ElementType,typename OctreeSemantics>
-void TOctree<ElementType,OctreeSemantics>::AddElementToNode(
+void TOctree2<ElementType,OctreeSemantics>::AddElementToNode(
 	typename TTypeTraits<ElementType>::ConstInitType Element,
 	const FNode& InNode,
 	const FOctreeNodeContext& InContext
@@ -765,7 +765,7 @@ void TOctree<ElementType,OctreeSemantics>::AddElementToNode(
 				// Create the child node if it hasn't been created yet.
 				if(!Node.Children[ChildRef.Index])
 				{
-					Node.Children[ChildRef.Index] = new typename TOctree<ElementType,OctreeSemantics>::FNode(&Node);
+					Node.Children[ChildRef.Index] = new typename TOctree2<ElementType,OctreeSemantics>::FNode(&Node);
 				}
 
 				// Push the node onto the stack to visit.
@@ -787,7 +787,7 @@ void TOctree<ElementType,OctreeSemantics>::AddElementToNode(
 }
 
 template<typename ElementType,typename OctreeSemantics>
-void TOctree<ElementType,OctreeSemantics>::DumpStats(bool bDetailed) const
+void TOctree2<ElementType,OctreeSemantics>::DumpStats(bool bDetailed) const
 {
 	int32 NumNodes = 0;
 	int32 NumLeaves = 0;
@@ -857,7 +857,7 @@ void TOctree<ElementType,OctreeSemantics>::DumpStats(bool bDetailed) const
 }
 
 template<typename ElementType,typename OctreeSemantics>
-void TOctree<ElementType,OctreeSemantics>::GetMemoryUsage(size_t& OutSizeBytes) const
+void TOctree2<ElementType,OctreeSemantics>::GetMemoryUsage(size_t& OutSizeBytes) const
 {
 	size_t SizeBytes = 0;
 
@@ -880,7 +880,7 @@ void TOctree<ElementType,OctreeSemantics>::GetMemoryUsage(size_t& OutSizeBytes)
 }
 
 template<typename ElementType,typename OctreeSemantics>
-TOctree<ElementType,OctreeSemantics>::TOctree(const FVector4& InOrigin,float InExtent)
+TOctree2<ElementType,OctreeSemantics>::TOctree(const FVector4& InOrigin,float InExtent)
 :	RootNode(NULL)
 ,	RootNodeContext(FBoxCenterAndExtent(InOrigin,FVector4(InExtent,InExtent,InExtent,0)), OctreeSemantics::LoosenessDenominator)
 ,	MinLeafExtent(InExtent * FMath::Pow((1.0f + 1.0f / (float)OctreeSemantics::LoosenessDenominator) / 2.0f,OctreeSemantics::MaxNodeDepth))
diff --git a/Engine/Source/Runtime/Core/Public/Math/GenericOctree.h b/Engine/Source/Runtime/Core/Public/Math/GenericOctree.h
--- a/Engine/Source/Runtime/Core/Public/Math/GenericOctree.h
+++ b/Engine/Source/Runtime/Core/Public/Math/GenericOctree.h
@@ -350,7 +350,7 @@ CORE_API DECLARE_LOG_CATEGORY_EXTERN(LogGenericOctree, Log, All);
 
 /** An octree. */
 template<typename ElementType,typename OctreeSemantics>
-class TOctree
+class TOctree2
 {
 	using ElementArrayType = TArray<ElementType, typename OctreeSemantics::ElementAllocator>;
 public:
@@ -986,7 +986,7 @@ public:
 	}
 
 	/** Initialization constructor. */
-	TOctree(const FVector& InOrigin,float InExtent)
+	TOctree2(const FVector& InOrigin,float InExtent)
 		: RootNodeContext(FBoxCenterAndExtent(InOrigin, FVector(InExtent, InExtent, InExtent)), 0, 0)
 		, MinLeafExtent(InExtent* FMath::Pow((1.0f + 1.0f / (float)FOctreeNodeContext::LoosenessDenominator) / 2.0f, OctreeSemantics::MaxNodeDepth))
 	{
@@ -995,7 +995,7 @@ public:
 	}
 
 	/** DO NOT USE. This constructor is for internal usage only for hot-reload purposes. */
-	TOctree() 
+	TOctree2() 
 	{
 		TreeNodes.AddDefaulted();
 		TreeElements.AddDefaulted();
diff --git a/Engine/Source/Runtime/Core/Public/Math/GenericOctreePublic.h b/Engine/Source/Runtime/Core/Public/Math/GenericOctreePublic.h
--- a/Engine/Source/Runtime/Core/Public/Math/GenericOctreePublic.h
+++ b/Engine/Source/Runtime/Core/Public/Math/GenericOctreePublic.h
@@ -18,6 +18,9 @@ public:
 	template<typename,typename>
 	friend class TOctree;
 
+	template<typename, typename>
+	friend class TOctree2;
+
 	/** Default constructor. */
 	FOctreeElementId()
 		:	NodeIndex(INDEX_NONE)
diff --git a/Engine/Source/Runtime/Engine/Private/StaticMeshBuild.cpp b/Engine/Source/Runtime/Engine/Private/StaticMeshBuild.cpp
--- a/Engine/Source/Runtime/Engine/Private/StaticMeshBuild.cpp
+++ b/Engine/Source/Runtime/Engine/Private/StaticMeshBuild.cpp
@@ -492,7 +492,7 @@ struct FStaticMeshComponentVertPosOctreeSemantics
 	{
 	}
 };
-typedef TOctree<FPaintedVertex, FStaticMeshComponentVertPosOctreeSemantics> TSMCVertPosOctree;
+typedef TOctree2<FPaintedVertex, FStaticMeshComponentVertPosOctreeSemantics> TSMCVertPosOctree;
 
 void RemapPaintedVertexColors(const TArray<FPaintedVertex>& InPaintedVertices,
 	const FColorVertexBuffer* InOverrideColors,
diff --git a/Engine/Source/Runtime/Engine/Public/PrecomputedLightVolume.h b/Engine/Source/Runtime/Engine/Public/PrecomputedLightVolume.h
--- a/Engine/Source/Runtime/Engine/Public/PrecomputedLightVolume.h
+++ b/Engine/Source/Runtime/Engine/Public/PrecomputedLightVolume.h
@@ -90,7 +90,7 @@ struct FLightVolumeOctreeSemantics
 	}
 };
 
-typedef TOctree<FVolumeLightingSample, FLightVolumeOctreeSemantics> FLightVolumeOctree;
+typedef TOctree2<FVolumeLightingSample, FLightVolumeOctreeSemantics> FLightVolumeOctree;
 
 /** Set of volume lighting samples belonging to one streaming level, which can be queried about the lighting at a given position. */
 class FPrecomputedLightVolumeData
diff --git a/Engine/Source/Runtime/Engine/Public/Rendering/SkeletalMeshLODImporterData.h b/Engine/Source/Runtime/Engine/Public/Rendering/SkeletalMeshLODImporterData.h
--- a/Engine/Source/Runtime/Engine/Public/Rendering/SkeletalMeshLODImporterData.h
+++ b/Engine/Source/Runtime/Engine/Public/Rendering/SkeletalMeshLODImporterData.h
@@ -609,7 +609,7 @@ struct FWedgeInfoOctreeSemantics
 	{
 	}
 };
-typedef TOctree<FWedgeInfo, FWedgeInfoOctreeSemantics> TWedgeInfoPosOctree;
+typedef TOctree2<FWedgeInfo, FWedgeInfoOctreeSemantics> TWedgeInfoPosOctree;
 
 class FOctreeQueryHelper
 {
diff --git a/Engine/Source/Runtime/Experimental/GeometryCollectionCore/Private/GeometryCollection/GeometryCollectionProximityUtility.cpp b/Engine/Source/Runtime/Experimental/GeometryCollectionCore/Private/GeometryCollection/GeometryCollectionProximityUtility.cpp
--- a/Engine/Source/Runtime/Experimental/GeometryCollectionCore/Private/GeometryCollection/GeometryCollectionProximityUtility.cpp
+++ b/Engine/Source/Runtime/Experimental/GeometryCollectionCore/Private/GeometryCollection/GeometryCollectionProximityUtility.cpp
@@ -69,7 +69,7 @@ struct FMeshProximityTriangleOctreeSemantics
 
 	}
 };
-typedef TOctree<FProximityTriangle, FMeshProximityTriangleOctreeSemantics> FProximityTriangleOctree;
+typedef TOctree2<FProximityTriangle, FMeshProximityTriangleOctreeSemantics> FProximityTriangleOctree;
 
 
 bool FGeometryCollectionProximityUtility::IsPointInsideOfTriangle(const FVector& P, const FVector& Vertex0, const FVector& Vertex1, const FVector& Vertex2, float Threshold)
diff --git a/Engine/Source/Runtime/NavigationSystem/Private/NavigationOctree.cpp b/Engine/Source/Runtime/NavigationSystem/Private/NavigationOctree.cpp
--- a/Engine/Source/Runtime/NavigationSystem/Private/NavigationOctree.cpp
+++ b/Engine/Source/Runtime/NavigationSystem/Private/NavigationOctree.cpp
@@ -9,7 +9,7 @@
 // FNavigationOctree
 //----------------------------------------------------------------------//
 FNavigationOctree::FNavigationOctree(const FVector& Origin, float Radius)
-	: TOctree<FNavigationOctreeElement, FNavigationOctreeSemantics>(Origin, Radius)
+	: TOctree2<FNavigationOctreeElement, FNavigationOctreeSemantics>(Origin, Radius)
 	, DefaultGeometryGatheringMode(ENavDataGatheringMode::Instant)
 	, bGatherGeometry(false)
 	, NodesMemory(0)
diff --git a/Engine/Source/Runtime/NavigationSystem/Public/NavigationOctree.h b/Engine/Source/Runtime/NavigationSystem/Public/NavigationOctree.h
--- a/Engine/Source/Runtime/NavigationSystem/Public/NavigationOctree.h
+++ b/Engine/Source/Runtime/NavigationSystem/Public/NavigationOctree.h
@@ -86,7 +86,7 @@ public:
 
 struct FNavigationOctreeSemantics
 {
-	typedef TOctree<FNavigationOctreeElement, FNavigationOctreeSemantics> FOctree;
+	typedef TOctree2<FNavigationOctreeElement, FNavigationOctreeSemantics> FOctree;
 	enum { MaxElementsPerLeaf = 16 };
 	enum { MinInclusiveElementsPerNode = 7 };
 	enum { MaxNodeDepth = 12 };
@@ -115,7 +115,7 @@ struct FNavigationOctreeSemantics
 	static void SetElementId(FOctree& OctreeOwner, const FNavigationOctreeElement& Element, FOctreeElementId Id);
 };
 
-class NAVIGATIONSYSTEM_API FNavigationOctree : public TOctree<FNavigationOctreeElement, FNavigationOctreeSemantics>, public TSharedFromThis<FNavigationOctree, ESPMode::ThreadSafe>
+class NAVIGATIONSYSTEM_API FNavigationOctree : public TOctree2<FNavigationOctreeElement, FNavigationOctreeSemantics>, public TSharedFromThis<FNavigationOctree, ESPMode::ThreadSafe>
 {
 public:
 	DECLARE_DELEGATE_TwoParams(FNavigableGeometryComponentExportDelegate, UActorComponent*, FNavigationRelevantData&);
@@ -134,7 +134,7 @@ public:
 	{
 		DEC_MEMORY_STAT_BY(STAT_NavigationMemory, OctreeSizeBytes);
 		DEC_MEMORY_STAT_BY(STAT_Navigation_CollisionTreeMemory, OctreeSizeBytes);
-		TOctree<FNavigationOctreeElement, FNavigationOctreeSemantics>::AddElement(Element);
+		TOctree2<FNavigationOctreeElement, FNavigationOctreeSemantics>::AddElement(Element);
 		OctreeSizeBytes = GetSizeBytes();
 		INC_MEMORY_STAT_BY(STAT_NavigationMemory, OctreeSizeBytes);
 		INC_MEMORY_STAT_BY(STAT_Navigation_CollisionTreeMemory, OctreeSizeBytes);
@@ -148,7 +148,7 @@ public:
 	{
 		DEC_MEMORY_STAT_BY(STAT_NavigationMemory, OctreeSizeBytes);
 		DEC_MEMORY_STAT_BY(STAT_Navigation_CollisionTreeMemory, OctreeSizeBytes);
-		TOctree<FNavigationOctreeElement, FNavigationOctreeSemantics>::RemoveElement(ElementId);
+		TOctree2<FNavigationOctreeElement, FNavigationOctreeSemantics>::RemoveElement(ElementId);
 		OctreeSizeBytes = GetSizeBytes();
 		INC_MEMORY_STAT_BY(STAT_NavigationMemory, OctreeSizeBytes);
 		INC_MEMORY_STAT_BY(STAT_Navigation_CollisionTreeMemory, OctreeSizeBytes);
diff --git a/Engine/Source/Runtime/Renderer/Private/LightSceneInfo.h b/Engine/Source/Runtime/Renderer/Private/LightSceneInfo.h
--- a/Engine/Source/Runtime/Renderer/Private/LightSceneInfo.h
+++ b/Engine/Source/Runtime/Renderer/Private/LightSceneInfo.h
@@ -143,7 +143,7 @@ struct TUseBitwiseSwap<FSortedLightSceneInfo>
 };
 
 /** The type of the octree used by FScene to find lights. */
-typedef TOctree<FLightSceneInfoCompact,struct FLightOctreeSemantics> FSceneLightOctree;
+typedef TOctree2<FLightSceneInfoCompact,struct FLightOctreeSemantics> FSceneLightOctree;
 
 /**
  * The information used to render a light.  This is the rendering thread's mirror of the game thread's ULightComponent.
diff --git a/Engine/Source/Runtime/Renderer/Public/PrimitiveSceneInfo.h b/Engine/Source/Runtime/Renderer/Public/PrimitiveSceneInfo.h
--- a/Engine/Source/Runtime/Renderer/Public/PrimitiveSceneInfo.h
+++ b/Engine/Source/Runtime/Renderer/Public/PrimitiveSceneInfo.h
@@ -207,7 +207,7 @@ struct FPrimitiveVirtualTextureLodInfo
 };
 
 /** The type of the octree used by FScene to find primitives. */
-typedef TOctree<FPrimitiveSceneInfoCompact,struct FPrimitiveOctreeSemantics> FScenePrimitiveOctree;
+typedef TOctree2<FPrimitiveSceneInfoCompact,struct FPrimitiveOctreeSemantics> FScenePrimitiveOctree;
 
 /**
  * The renderer's internal state for a single UPrimitiveComponent.  This has a one to one mapping with FPrimitiveSceneProxy, which is in the engine module.
