Fix for lightmass assert in debug mode.

#rnx

diff --git a/Engine/Source/Programs/UnrealLightmass/Private/LightmassCore/Math/LMOctree.h b/Engine/Source/Programs/UnrealLightmass/Private/LightmassCore/Math/LMOctree.h
--- a/Engine/Source/Programs/UnrealLightmass/Private/LightmassCore/Math/LMOctree.h
+++ b/Engine/Source/Programs/UnrealLightmass/Private/LightmassCore/Math/LMOctree.h
@@ -617,7 +617,7 @@ public:
 
 private:
 	/** Lightmass octree is grow-only. Use a stack allocator to coalesce allocations for better allocation performance. */
-	FMemStackBase NodeAllocator;
+	FMemStackBase NodeAllocator = FMemStackBase(0);
 
 	/** The octree's root node. */
 	FNode RootNode;
