Removing unused d3d9 code from Lightmass

diff --git a/Engine/Source/Programs/UnrealLightmass/Private/ImportExport/3DVisualizer.cpp b/Engine/Source/Programs/UnrealLightmass/Private/ImportExport/3DVisualizer.cpp
deleted file mode 100644
diff --git a/Engine/Source/Programs/UnrealLightmass/Private/ImportExport/3DVisualizer.h b/Engine/Source/Programs/UnrealLightmass/Private/ImportExport/3DVisualizer.h
deleted file mode 100644
diff --git a/Engine/Source/Programs/UnrealLightmass/Private/Launch/UnrealLightmass.cpp b/Engine/Source/Programs/UnrealLightmass/Private/Launch/UnrealLightmass.cpp
--- a/Engine/Source/Programs/UnrealLightmass/Private/Launch/UnrealLightmass.cpp
+++ b/Engine/Source/Programs/UnrealLightmass/Private/Launch/UnrealLightmass.cpp
@@ -18,12 +18,6 @@
 #include "IMessagingModule.h"
 #endif
 
-#if PLATFORM_WINDOWS
-#include "Windows/AllowWindowsPlatformTypes.h"
-	#include <d3dx9.h>
-#include "Windows/HideWindowsPlatformTypes.h"
-#endif
-
 DEFINE_LOG_CATEGORY(LogLightmass);
 
 IMPLEMENT_APPLICATION(UnrealLightmass, "UnrealLightmass");
@@ -420,31 +414,6 @@ void CriticalErrorCallback()
 }
 
 #if PLATFORM_WINDOWS
-/**
- * Verifies that the correct version of DirectX is installed.
- * @return	true if everything looks correct
- */
-bool VerifyD3D()
-{
-	bool bD3DInstalledCorrectly = false;
-	IDirect3D9* D3D = NULL;
-	__try
-	{
-		D3D = Direct3DCreate9(D3D_SDK_VERSION);
-		bD3DInstalledCorrectly = (D3D != NULL) && D3DXCheckVersion(D3D_SDK_VERSION, D3DX_SDK_VERSION);
-	}
-	__except( EXCEPTION_EXECUTE_HANDLER )
-	{
-		bD3DInstalledCorrectly = false;
-	}
-	if ( !bD3DInstalledCorrectly )
-	{
-		UE_LOG(LogLightmass, Display,  TEXT("DirectX run-time isn't installed or it's using the incorrect version!\nLightmass requires D3D_SDK_VERSION %d and D3DX_SDK_VERSION %d."), D3D_SDK_VERSION, D3DX_SDK_VERSION );
-		return false;
-	}
-	D3D->Release();
-	return true;
-}
 
 bool VerifyDLL( const TCHAR* DLLFilename )
 {
@@ -483,9 +452,8 @@ int main(int argc, ANSICHAR* argv[])
 	// Set the error mode to avoid popping up dialog boxes on crashes
 	SetErrorMode( SEM_NOGPFAULTERRORBOX | SEM_NOGPFAULTERRORBOX );
 
-	// Verify the installed DirectX run-time and other required DLLs
-	if ( !Lightmass::VerifyD3D() ||
-		 !Lightmass::VerifyDLL(TEXT("dbghelp.dll")) )
+	// Verify the required DLLs
+	if ( !Lightmass::VerifyDLL(TEXT("dbghelp.dll")) )
 	{
 		return 1;
 	}
diff --git a/Engine/Source/Programs/UnrealLightmass/UnrealLightmass.Build.cs b/Engine/Source/Programs/UnrealLightmass/UnrealLightmass.Build.cs
--- a/Engine/Source/Programs/UnrealLightmass/UnrealLightmass.Build.cs
+++ b/Engine/Source/Programs/UnrealLightmass/UnrealLightmass.Build.cs
@@ -14,8 +14,6 @@ public class UnrealLightmass : ModuleRules
 
 		if ((Target.Platform == UnrealTargetPlatform.Win64) || (Target.Platform == UnrealTargetPlatform.Win32))
 		{
-			AddEngineThirdPartyPrivateStaticDependencies(Target, "DX9");
-
 			// Unreallightmass requires GetProcessMemoryInfo exported by psapi.dll. http://msdn.microsoft.com/en-us/library/windows/desktop/ms683219(v=vs.85).aspx
 			PublicSystemLibraries.Add("psapi.lib");
 			PrivateDependencyModuleNames.AddRange(
