Double math ambiguity Engine further fixes

diff --git a/Engine/Source/Programs/UnrealLightmass/Private/ImportExport/LightmassScene.cpp b/Engine/Source/Programs/UnrealLightmass/Private/ImportExport/LightmassScene.cpp
--- a/Engine/Source/Programs/UnrealLightmass/Private/ImportExport/LightmassScene.cpp
+++ b/Engine/Source/Programs/UnrealLightmass/Private/ImportExport/LightmassScene.cpp
@@ -1829,7 +1829,7 @@ void FSkyLight::Import( FLightmassImporter& Importer )
 	TArray<FFloat16Color> RadianceEnvironmentMap;
 	Importer.ImportArray(RadianceEnvironmentMap, RadianceEnvironmentMapDataSize);
 
-	CubemapSize = FMath::Sqrt(RadianceEnvironmentMapDataSize / 6);
+	CubemapSize = FMath::Sqrt(RadianceEnvironmentMapDataSize / 6.f);
 	NumMips = FMath::CeilLogTwo(CubemapSize) + 1;
 
 	check(FMath::IsPowerOfTwo(CubemapSize));
