kDOPTree - Fix lost precision in mean/variance computation when accumulating using a float
Would end up crashing with an OoB access as the code would put all triangles on the left of the splitting plane, which was not expected

Also optimized SplitTriangleList() to avoid going through the triangle list 6x...

diff --git a/Engine/Source/Programs/UnrealLightmass/Private/LightmassCore/Math/LMkDOP.h b/Engine/Source/Programs/UnrealLightmass/Private/LightmassCore/Math/LMkDOP.h
--- a/Engine/Source/Programs/UnrealLightmass/Private/LightmassCore/Math/LMkDOP.h
+++ b/Engine/Source/Programs/UnrealLightmass/Private/LightmassCore/Math/LMkDOP.h
@@ -1159,39 +1159,59 @@ struct TkDOPTree
 			int32 BestPlane = -1;
 			float BestMean = 0.f;
 			float BestVariance = 0.f;
+
 			// Determine how to split using the splatter algorithm
-			for (int32 nPlane = 0; nPlane < NUM_PLANES; nPlane++)
 			{
-				float Mean = 0.f;
-				float Variance = 0.f;
+				double Mean[NUM_PLANES] = { 0 };
+				double Variance[NUM_PLANES] = { 0 };
+
 				// Compute the mean for the triangle list
 				for (int32 nTriangle = Start; nTriangle < Start + NumTris; nTriangle++)
 				{
 					// Project the centroid of the triangle against the plane
 					// normals and accumulate to find the total projected
 					// weighting
-					Mean += BuildTriangles[nTriangle].GetCentroid()[nPlane];
+					FVector4 Centroid = BuildTriangles[nTriangle].GetCentroid();
+
+					for (int32 nPlane = 0; nPlane < NUM_PLANES; nPlane++)
+					{
+						Mean[nPlane] += Centroid[nPlane];
+					}
 				}
+
 				// Divide by the number of triangles to get the average
-				Mean /= float(NumTris);
+				for (int32 nPlane = 0; nPlane < NUM_PLANES; nPlane++)
+				{
+					Mean[nPlane] /= NumTris;
+				}
+
 				// Compute variance of the triangle list
-				for (int32 nTriangle = Start; nTriangle < Start + NumTris;nTriangle++)
+				for (int32 nTriangle = Start; nTriangle < Start + NumTris; nTriangle++)
 				{
 					// Project the centroid again
-					float Dot = BuildTriangles[nTriangle].GetCentroid()[nPlane];
+					FVector4 Centroid = BuildTriangles[nTriangle].GetCentroid();
+
 					// Now calculate the variance and accumulate it
-					Variance += (Dot - Mean) * (Dot - Mean);
+					for (int32 nPlane = 0; nPlane < NUM_PLANES; nPlane++)
+					{
+						Variance[nPlane] += (Centroid[nPlane] - Mean[nPlane]) * (Centroid[nPlane] - Mean[nPlane]);
+					}
 				}
-				// Get the average variance
-				Variance /= float(NumTris);
-				// Determine if this plane is the best to split on or not
-				if (Variance >= BestVariance)
+
+				// Determine which plane is the best to split on
+				for (int32 nPlane = 0; nPlane < NUM_PLANES; nPlane++)
 				{
-					BestPlane = nPlane;
-					BestVariance = Variance;
-					BestMean = Mean;
+					// Get the average variance
+					Variance[nPlane] /= NumTris;
+					if (Variance[nPlane] >= BestVariance)
+					{
+						BestPlane = nPlane;
+						BestVariance = (float)Variance[nPlane];
+						BestMean = (float)Mean[nPlane];
+					}
 				}
 			}
+
 			// Now that we have the plane to split on, work through the triangle
 			// list placing them on the left or right of the splitting plane
 			int32 Left = Start - 1;
diff --git a/Engine/Source/Runtime/Engine/Public/kDOP.h b/Engine/Source/Runtime/Engine/Public/kDOP.h
--- a/Engine/Source/Runtime/Engine/Public/kDOP.h
+++ b/Engine/Source/Runtime/Engine/Public/kDOP.h
@@ -444,39 +444,59 @@ struct TkDOPNode
 			int32 BestPlane = -1;
 			float BestMean = 0.f;
 			float BestVariance = 0.f;
+
 			// Determine how to split using the splatter algorithm
-			for (int32 nPlane = 0; nPlane < NUM_PLANES; nPlane++)
 			{
-				float Mean = 0.f;
-				float Variance = 0.f;
+				double Mean[NUM_PLANES] = { 0 };
+				double Variance[NUM_PLANES] = { 0 };
+
 				// Compute the mean for the triangle list
 				for (int32 nTriangle = Start; nTriangle < Start + NumTris; nTriangle++)
 				{
 					// Project the centroid of the triangle against the plane
 					// normals and accumulate to find the total projected
 					// weighting
-					Mean += BuildTriangles[nTriangle].GetCentroid()[nPlane];
+					FVector4 Centroid = BuildTriangles[nTriangle].GetCentroid();
+
+					for (int32 nPlane = 0; nPlane < NUM_PLANES; nPlane++)
+					{
+						Mean[nPlane] += Centroid[nPlane];
+					}
 				}
+
 				// Divide by the number of triangles to get the average
-				Mean /= float(NumTris);
+				for (int32 nPlane = 0; nPlane < NUM_PLANES; nPlane++)
+				{
+					Mean[nPlane] /= NumTris;
+				}
+
 				// Compute variance of the triangle list
-				for (int32 nTriangle = Start; nTriangle < Start + NumTris;nTriangle++)
+				for (int32 nTriangle = Start; nTriangle < Start + NumTris; nTriangle++)
 				{
 					// Project the centroid again
-					float Dot = BuildTriangles[nTriangle].GetCentroid()[nPlane];
+					FVector4 Centroid = BuildTriangles[nTriangle].GetCentroid();
+
 					// Now calculate the variance and accumulate it
-					Variance += (Dot - Mean) * (Dot - Mean);
+					for (int32 nPlane = 0; nPlane < NUM_PLANES; nPlane++)
+					{
+						Variance[nPlane] += (Centroid[nPlane] - Mean[nPlane]) * (Centroid[nPlane] - Mean[nPlane]);
+					}
 				}
-				// Get the average variance
-				Variance /= float(NumTris);
-				// Determine if this plane is the best to split on or not
-				if (Variance >= BestVariance)
+
+				// Determine which plane is the best to split on
+				for (int32 nPlane = 0; nPlane < NUM_PLANES; nPlane++)
 				{
-					BestPlane = nPlane;
-					BestVariance = Variance;
-					BestMean = Mean;
+					// Get the average variance
+					Variance[nPlane] /= NumTris;
+					if (Variance[nPlane] >= BestVariance)
+					{
+						BestPlane = nPlane;
+						BestVariance = (float)Variance[nPlane];
+						BestMean = (float)Mean[nPlane];
+					}
 				}
 			}
+
 			// Now that we have the plane to split on, work through the triangle
 			// list placing them on the left or right of the splitting plane
 			int32 Left = Start - 1;
diff --git a/Engine/Source/Runtime/Experimental/Chaos/Private/kDOP.h b/Engine/Source/Runtime/Experimental/Chaos/Private/kDOP.h
--- a/Engine/Source/Runtime/Experimental/Chaos/Private/kDOP.h
+++ b/Engine/Source/Runtime/Experimental/Chaos/Private/kDOP.h
@@ -444,39 +444,59 @@ struct TkDOPNode
 			int32 BestPlane = -1;
 			float BestMean = 0.f;
 			float BestVariance = 0.f;
+
 			// Determine how to split using the splatter algorithm
-			for (int32 nPlane = 0; nPlane < NUM_PLANES; nPlane++)
 			{
-				float Mean = 0.f;
-				float Variance = 0.f;
+				double Mean[NUM_PLANES] = { 0 };
+				double Variance[NUM_PLANES] = { 0 };
+
 				// Compute the mean for the triangle list
 				for (int32 nTriangle = Start; nTriangle < Start + NumTris; nTriangle++)
 				{
 					// Project the centroid of the triangle against the plane
 					// normals and accumulate to find the total projected
 					// weighting
-					Mean += BuildTriangles[nTriangle].GetCentroid()[nPlane];
+					FVector4 Centroid = BuildTriangles[nTriangle].GetCentroid();
+
+					for (int32 nPlane = 0; nPlane < NUM_PLANES; nPlane++)
+					{
+						Mean[nPlane] += Centroid[nPlane];
+					}
 				}
+
 				// Divide by the number of triangles to get the average
-				Mean /= float(NumTris);
+				for (int32 nPlane = 0; nPlane < NUM_PLANES; nPlane++)
+				{
+					Mean[nPlane] /= NumTris;
+				}
+
 				// Compute variance of the triangle list
-				for (int32 nTriangle = Start; nTriangle < Start + NumTris;nTriangle++)
+				for (int32 nTriangle = Start; nTriangle < Start + NumTris; nTriangle++)
 				{
 					// Project the centroid again
-					float Dot = BuildTriangles[nTriangle].GetCentroid()[nPlane];
+					FVector4 Centroid = BuildTriangles[nTriangle].GetCentroid();
+
 					// Now calculate the variance and accumulate it
-					Variance += (Dot - Mean) * (Dot - Mean);
+					for (int32 nPlane = 0; nPlane < NUM_PLANES; nPlane++)
+					{
+						Variance[nPlane] += (Centroid[nPlane] - Mean[nPlane]) * (Centroid[nPlane] - Mean[nPlane]);
+					}
 				}
-				// Get the average variance
-				Variance /= float(NumTris);
-				// Determine if this plane is the best to split on or not
-				if (Variance >= BestVariance)
+
+				// Determine which plane is the best to split on
+				for (int32 nPlane = 0; nPlane < NUM_PLANES; nPlane++)
 				{
-					BestPlane = nPlane;
-					BestVariance = Variance;
-					BestMean = Mean;
+					// Get the average variance
+					Variance[nPlane] /= NumTris;
+					if (Variance[nPlane] >= BestVariance)
+					{
+						BestPlane = nPlane;
+						BestVariance = (float)Variance[nPlane];
+						BestMean = (float)Mean[nPlane];
+					}
 				}
 			}
+
 			// Now that we have the plane to split on, work through the triangle
 			// list placing them on the left or right of the splitting plane
 			int32 Left = Start - 1;
