Fix incorrect depends path

diff --git a/Engine/Source/Programs/UnrealLightmass/UnrealLightmass.Build.cs b/Engine/Source/Programs/UnrealLightmass/UnrealLightmass.Build.cs
--- a/Engine/Source/Programs/UnrealLightmass/UnrealLightmass.Build.cs
+++ b/Engine/Source/Programs/UnrealLightmass/UnrealLightmass.Build.cs
@@ -89,7 +89,7 @@ public class UnrealLightmass : ModuleRules
 
             PublicIncludePaths.Add(Path.Combine(SDKDir, "include"));
             PublicAdditionalLibraries.Add(Path.Combine(SDKDir, "lib/libembree.so"));
-			RuntimeDependencies.Add("$(EngineDir)/Binaries/Linux/libembree.so");
+			RuntimeDependencies.Add(Path.Combine(SDKDir, "lib/libembree.so"));
 			PublicAdditionalLibraries.Add(Path.Combine(IntelTBBLibs, "Linux/libtbb.so"));
 			RuntimeDependencies.Add(Path.Combine(IntelTBBLibs, "Linux/libtbb.so"));
 			RuntimeDependencies.Add(Path.Combine(IntelTBBLibs, "Linux/libtbb.so.2"));
