SplitTriangleList: Fix loop conditions to avoid off-by-one dereference of end of the list.
#rnx

diff --git a/Engine/Source/Programs/UnrealLightmass/Private/LightmassCore/Math/LMkDOP.h b/Engine/Source/Programs/UnrealLightmass/Private/LightmassCore/Math/LMkDOP.h
--- a/Engine/Source/Programs/UnrealLightmass/Private/LightmassCore/Math/LMkDOP.h
+++ b/Engine/Source/Programs/UnrealLightmass/Private/LightmassCore/Math/LMkDOP.h
@@ -7,6 +7,7 @@
 #include "LMMath.h"
 #include "LMMathSSE.h"
 #include "LMStats.h"
+#include "Templates/UnrealTemplate.h"
 #include "UnrealLightmass.h"
 
 namespace Lightmass
@@ -1216,32 +1217,43 @@ struct TkDOPTree
 			// list placing them on the left or right of the splitting plane
 			int32 Left = Start - 1;
 			int32 Right = Start + NumTris;
-			// Keep working through until the left index passes the right
-			while (Left < Right)
+			// Keep working through until the left index passes the right; we test that condition after each interior loop
+			for (;;)
 			{
+				// Loop invariants: (1) Left < Right,
+				// (2) All triangles <= Left belong on the left, all triangles >= Right belong on the right
+				// (3) Left+1 is an untested triangle, Right-1 is an untested triangle
 				float Dot;
-				// Find all the triangles to the "left" of the splitting plane
-				do
+				// Increment Left until it points to triangle that belongs on the right, or Right==Left
+				for (++Left; Left < Right; ++Left)
 				{
-					Dot = BuildTriangles[++Left].GetCentroid()[BestPlane];
+					Dot = BuildTriangles[Left].GetCentroid()[BestPlane];
+					if (Dot < BestMean)
+					{
+						break;
+					}
+				}
+				if (Left == Right)
+				{
+					break;
 				}
-				while (Dot < BestMean && Left < Right);
-				// Find all the triangles to the "right" of the splitting plane
-				do
+				// Decrement Right until it points to triangle that belongs on the left, or Right==Left
+				for (--Right; Left < Right; --Right)
 				{
-					Dot = BuildTriangles[--Right].GetCentroid()[BestPlane];
+					Dot = BuildTriangles[Right].GetCentroid()[BestPlane];
+					if (Dot >= BestMean)
+					{
+						break;
+					}
 				}
-				while (Dot >= BestMean && Right > 0 && Left < Right);
-				// Don't swap the triangle data if we just hit the end
-				if (Left < Right)
+				if (Left == Right)
 				{
-					// Swap the triangles since they are on the wrong sides of the
-					// splitting plane
-					FkDOPBuildCollisionTriangle<KDOP_IDX_TYPE> Temp = BuildTriangles[Left];
-					BuildTriangles[Left] = BuildTriangles[Right];
-					BuildTriangles[Right] = Temp;
+					break;
 				}
+				// Left points to a triangle that belongs on the Right; Right points to a triangle that belongs on the Left. Swap them.
+				Swap(BuildTriangles[Left], BuildTriangles[Right]);
 			}
+			// After loop array is partitioned and Left is the first index that belongs on the right side of the plane.
 			// Check for wacky degenerate case where more than GKDOPMaxTrisPerLeaf
 			// fall all in the same kDOP
 			if (Left == Start + NumTris || Right == Start)
diff --git a/Engine/Source/Runtime/Engine/Public/kDOP.h b/Engine/Source/Runtime/Engine/Public/kDOP.h
--- a/Engine/Source/Runtime/Engine/Public/kDOP.h
+++ b/Engine/Source/Runtime/Engine/Public/kDOP.h
@@ -4,6 +4,7 @@
 
 #include "CoreMinimal.h"
 #include "Containers/ChunkedArray.h"
+#include "Templates/UnrealTemplate.h"
 
 // Indicates how many "k / 2" there are in the k-DOP. 3 == AABB == 6 DOP. The code relies on this being 3.
 #define NUM_PLANES	3
@@ -501,32 +502,43 @@ struct TkDOPNode
 			// list placing them on the left or right of the splitting plane
 			int32 Left = Start - 1;
 			int32 Right = Start + NumTris;
-			// Keep working through until the left index passes the right
-			while (Left < Right)
+			// Keep working through until the left index passes the right; we test that condition after each interior loop
+			for (;;)
 			{
+				// Loop invariants: (1) Left < Right,
+				// (2) All triangles <= Left belong on the left, all triangles >= Right belong on the right
+				// (3) Left+1 is an untested triangle, Right-1 is an untested triangle
 				float Dot;
-				// Find all the triangles to the "left" of the splitting plane
-				do
+				// Increment Left until it points to triangle that belongs on the right, or Right==Left
+				for (++Left; Left < Right; ++Left)
 				{
-					Dot = BuildTriangles[++Left].GetCentroid()[BestPlane];
+					Dot = BuildTriangles[Left].GetCentroid()[BestPlane];
+					if (Dot < BestMean)
+					{
+						break;
+					}
+				}
+				if (Left == Right)
+				{
+					break;
 				}
-				while (Dot < BestMean && Left < Right);
-				// Find all the triangles to the "right" of the splitting plane
-				do
+				// Decrement Right until it points to triangle that belongs on the left, or Right==Left
+				for (--Right; Left < Right; --Right)
 				{
-					Dot = BuildTriangles[--Right].GetCentroid()[BestPlane];
+					Dot = BuildTriangles[Right].GetCentroid()[BestPlane];
+					if (Dot >= BestMean)
+					{
+						break;
+					}
 				}
-				while (Dot >= BestMean && Right > 0 && Left < Right);
-				// Don't swap the triangle data if we just hit the end
-				if (Left < Right)
+				if (Left == Right)
 				{
-					// Swap the triangles since they are on the wrong sides of the
-					// splitting plane
-					FkDOPBuildCollisionTriangle<KDOP_IDX_TYPE> Temp = BuildTriangles[Left];
-					BuildTriangles[Left] = BuildTriangles[Right];
-					BuildTriangles[Right] = Temp;
+					break;
 				}
+				// Left points to a triangle that belongs on the Right; Right points to a triangle that belongs on the Left. Swap them.
+				Swap(BuildTriangles[Left], BuildTriangles[Right]);
 			}
+			// After loop array is partitioned and Left is the first index that belongs on the right side of the plane.
 			// Check for wacky degenerate case where more than GKDOPMaxTrisPerLeaf
 			// fall all in the same kDOP
 			if (Left == Start + NumTris || Right == Start)
diff --git a/Engine/Source/Runtime/Experimental/Chaos/Private/kDOP.h b/Engine/Source/Runtime/Experimental/Chaos/Private/kDOP.h
--- a/Engine/Source/Runtime/Experimental/Chaos/Private/kDOP.h
+++ b/Engine/Source/Runtime/Experimental/Chaos/Private/kDOP.h
@@ -4,6 +4,7 @@
 
 #include "CoreMinimal.h"
 #include "Containers/ChunkedArray.h"
+#include "Templates/UnrealTemplate.h"
 
 // Indicates how many "k / 2" there are in the k-DOP. 3 == AABB == 6 DOP. The code relies on this being 3.
 #define NUM_PLANES	3
@@ -501,32 +502,43 @@ struct TkDOPNode
 			// list placing them on the left or right of the splitting plane
 			int32 Left = Start - 1;
 			int32 Right = Start + NumTris;
-			// Keep working through until the left index passes the right
-			while (Left < Right)
+			// Keep working through until the left index passes the right; we test that condition after each interior loop
+			for (;;)
 			{
+				// Loop invariants: (1) Left < Right,
+				// (2) All triangles <= Left belong on the left, all triangles >= Right belong on the right
+				// (3) Left+1 is an untested triangle, Right-1 is an untested triangle
 				float Dot;
-				// Find all the triangles to the "left" of the splitting plane
-				do
+				// Increment Left until it points to triangle that belongs on the right, or Right==Left
+				for (++Left; Left < Right; ++Left)
 				{
-					Dot = BuildTriangles[++Left].GetCentroid()[BestPlane];
+					Dot = BuildTriangles[Left].GetCentroid()[BestPlane];
+					if (Dot < BestMean)
+					{
+						break;
+					}
+				}
+				if (Left == Right)
+				{
+					break;
 				}
-				while (Dot < BestMean && Left < Right);
-				// Find all the triangles to the "right" of the splitting plane
-				do
+				// Decrement Right until it points to triangle that belongs on the left, or Right==Left
+				for (--Right; Left < Right; --Right)
 				{
-					Dot = BuildTriangles[--Right].GetCentroid()[BestPlane];
+					Dot = BuildTriangles[Right].GetCentroid()[BestPlane];
+					if (Dot >= BestMean)
+					{
+						break;
+					}
 				}
-				while (Dot >= BestMean && Right > 0 && Left < Right);
-				// Don't swap the triangle data if we just hit the end
-				if (Left < Right)
+				if (Left == Right)
 				{
-					// Swap the triangles since they are on the wrong sides of the
-					// splitting plane
-					FkDOPBuildCollisionTriangle<KDOP_IDX_TYPE> Temp = BuildTriangles[Left];
-					BuildTriangles[Left] = BuildTriangles[Right];
-					BuildTriangles[Right] = Temp;
+					break;
 				}
+				// Left points to a triangle that belongs on the Right; Right points to a triangle that belongs on the Left. Swap them.
+				Swap(BuildTriangles[Left], BuildTriangles[Right]);
 			}
+			// After loop array is partitioned and Left is the first index that belongs on the right side of the plane.
 			// Check for wacky degenerate case where more than GKDOPMaxTrisPerLeaf
 			// fall all in the same kDOP
 			if (Left == Start + NumTris || Right == Start)
