remove duplicate Color::ComputeLuminance , use GetLuminance
also fix bug in ProxyLOD AverageColor , iterating on Size instead of 4

diff --git a/Engine/Plugins/Experimental/ProxyLODPlugin/Source/ProxyLOD/Private/ProxyLODMaterialTransferUtilities.cpp b/Engine/Plugins/Experimental/ProxyLODPlugin/Source/ProxyLOD/Private/ProxyLODMaterialTransferUtilities.cpp
--- a/Engine/Plugins/Experimental/ProxyLODPlugin/Source/ProxyLOD/Private/ProxyLODMaterialTransferUtilities.cpp
+++ b/Engine/Plugins/Experimental/ProxyLODPlugin/Source/ProxyLOD/Private/ProxyLODMaterialTransferUtilities.cpp
@@ -877,7 +877,7 @@ namespace
 			}
 			else
 			{
-				Result = Color.ComputeLuminance();
+				Result = Color.GetLuminance();
 			}
 
 			return Result;
diff --git a/Engine/Plugins/Experimental/ProxyLODPlugin/Source/ProxyLOD/Private/ProxyLODMeshAttrTransfer.cpp b/Engine/Plugins/Experimental/ProxyLODPlugin/Source/ProxyLOD/Private/ProxyLODMeshAttrTransfer.cpp
--- a/Engine/Plugins/Experimental/ProxyLODPlugin/Source/ProxyLOD/Private/ProxyLODMeshAttrTransfer.cpp
+++ b/Engine/Plugins/Experimental/ProxyLODPlugin/Source/ProxyLOD/Private/ProxyLODMeshAttrTransfer.cpp
@@ -2,13 +2,12 @@
 
 #include "ProxyLODMeshAttrTransfer.h"
 #include "ProxyLODBarycentricUtilities.h" // for some of the barycentric stuff
-
+#include "Templates/CheckValueCast.h"
 
 
 template<int Size>
 FColor AverageColor(const FColor(&Colors)[Size])
 {
-	FColor Result;
 	// Accumulate with floats because FColor is only 8-bit
 	float Tmp[4] = { 0,0,0,0 };
 	for (int i = 0; i < Size; ++i)
@@ -18,12 +17,13 @@ FColor AverageColor(const FColor(&Colors)[Size])
 		Tmp[2] += Colors[i].B;
 		Tmp[3] += Colors[i].A;
 	}
-	for (int i = 0; i < Size; ++i) Tmp[i] *= 1.f / float(Size);
-
-	Result.R = Tmp[0];
-	Result.G = Tmp[1];
-	Result.B = Tmp[2];
-	Result.A = Tmp[3];
+	for (int i = 0; i < 4; ++i) Tmp[i] *= 1.f / float(Size);
+	
+	FColor Result;
+	Result.R = TCheckValueCast<uint8>( (int)(0.5f + Tmp[0] ) );
+	Result.G = TCheckValueCast<uint8>( (int)(0.5f + Tmp[1] ) );
+	Result.B = TCheckValueCast<uint8>( (int)(0.5f + Tmp[2] ) );
+	Result.A = TCheckValueCast<uint8>( (int)(0.5f + Tmp[3] ) );
 
 	return Result;
 }
@@ -242,10 +242,10 @@ void ProxyLOD::TransferVertexColors(const FClosestPolyField& SrcPolyField, FMesh
 
 					FLinearColor InterpolatedColor = InterpolateVertexData(Weights, WedgeColors);
 
-					float AveLum = WedgeColors[0].ComputeLuminance() + WedgeColors[1].ComputeLuminance() + WedgeColors[2].ComputeLuminance();
+					float AveLum = WedgeColors[0].GetLuminance() + WedgeColors[1].GetLuminance() + WedgeColors[2].GetLuminance();
 					AveLum /= 3.f;
 
-					float LumTmp = InterpolatedColor.ComputeLuminance();
+					float LumTmp = InterpolatedColor.GetLuminance();
 					if (LumTmp > 1.e-5 && AveLum > 1.e-5)
 					{
 						InterpolatedColor *= AveLum / LumTmp;
diff --git a/Engine/Source/Developer/TextureCompressor/Private/TextureCompressorModule.cpp b/Engine/Source/Developer/TextureCompressor/Private/TextureCompressorModule.cpp
--- a/Engine/Source/Developer/TextureCompressor/Private/TextureCompressorModule.cpp
+++ b/Engine/Source/Developer/TextureCompressor/Private/TextureCompressorModule.cpp
@@ -473,7 +473,7 @@ static void GenerateSharpenedMipB8G8R8A8Templ(
 					}
 				}
 
-				float NewLuminance = SharpenedColor.ComputeLuminance();
+				float NewLuminance = SharpenedColor.GetLuminance();
 
 				// simple 2x2 kernel to compute the color
 				FilteredColor =
@@ -482,7 +482,7 @@ static void GenerateSharpenedMipB8G8R8A8Templ(
 					+ LookupSourceMip<AddressMode>( SourceImageData, SourceX + 0, SourceY + 1 )
 					+ LookupSourceMip<AddressMode>( SourceImageData, SourceX + 1, SourceY + 1 ) ) * 0.25f;
 
-				float OldLuminance = FilteredColor.ComputeLuminance();
+				float OldLuminance = FilteredColor.GetLuminance();
 
 				if ( OldLuminance > 0.001f )
 				{
diff --git a/Engine/Source/Developer/TraceInsights/Private/Insights/NetworkingProfiler/ViewModels/PacketContentViewDrawHelper.cpp b/Engine/Source/Developer/TraceInsights/Private/Insights/NetworkingProfiler/ViewModels/PacketContentViewDrawHelper.cpp
--- a/Engine/Source/Developer/TraceInsights/Private/Insights/NetworkingProfiler/ViewModels/PacketContentViewDrawHelper.cpp
+++ b/Engine/Source/Developer/TraceInsights/Private/Insights/NetworkingProfiler/ViewModels/PacketContentViewDrawHelper.cpp
@@ -256,7 +256,7 @@ void FPacketContentViewDrawStateBuilder::AddEvent(const TraceServices::FNetProfi
 			// Grey threshold is shifted toward black (0.4 instead of 0.5 in test below) due to "area rule":
 			// a large gray surface (background of an event in this case) is perceived lighter than a smaller area (text pixels).
 			// Ref: https://books.google.ro/books?id=0pVr7dhmdWYC
-			const bool bIsDarkColor = (EventColorFill.ComputeLuminance() < 0.4f);
+			const bool bIsDarkColor = (EventColorFill.GetLuminance() < 0.4f);
 
 			DrawState.Texts.AddDefaulted();
 			FPacketContentViewDrawState::FTextPrimitive& DrawText = DrawState.Texts.Last();
diff --git a/Engine/Source/Developer/TraceInsights/Private/Insights/ViewModels/TimingViewDrawHelper.cpp b/Engine/Source/Developer/TraceInsights/Private/Insights/ViewModels/TimingViewDrawHelper.cpp
--- a/Engine/Source/Developer/TraceInsights/Private/Insights/ViewModels/TimingViewDrawHelper.cpp
+++ b/Engine/Source/Developer/TraceInsights/Private/Insights/ViewModels/TimingViewDrawHelper.cpp
@@ -218,7 +218,7 @@ void FTimingEventsTrackDrawStateBuilder::AddEvent(double EventStartTime, double
 			// Grey threshold is shifted toward black (0.4 instead of 0.5 in test below) due to "area rule":
 			// a large gray surface (background of a timing event in this case) is perceived lighter than a smaller area (text pixels).
 			// Ref: https://books.google.ro/books?id=0pVr7dhmdWYC
-			const bool bIsDarkColor = (EventColorFill.ComputeLuminance() < 0.4f);
+			const bool bIsDarkColor = (EventColorFill.GetLuminance() < 0.4f);
 
 			DrawState.Texts.AddDefaulted();
 			FTimingEventsTrackDrawState::FTextPrimitive& DrawText = DrawState.Texts.Last();
diff --git a/Engine/Source/Editor/EnvironmentLightingViewer/Private/SEnvironmentLightingViewer.cpp b/Engine/Source/Editor/EnvironmentLightingViewer/Private/SEnvironmentLightingViewer.cpp
--- a/Engine/Source/Editor/EnvironmentLightingViewer/Private/SEnvironmentLightingViewer.cpp
+++ b/Engine/Source/Editor/EnvironmentLightingViewer/Private/SEnvironmentLightingViewer.cpp
@@ -268,7 +268,7 @@ void SEnvironmentLightingViewer::Tick(const FGeometry& AllottedGeometry, const d
 					/*|| AtmosphericLight->*/)
 					continue;
 
-				float LightLuminance = AtmosphericLight->GetColoredLightBrightness().ComputeLuminance();
+				float LightLuminance = AtmosphericLight->GetColoredLightBrightness().GetLuminance();
 				if (!SelectedAtmosphericLight ||					// Set it if null
 					SelectedLightLuminance < LightLuminance)		// Or choose the brightest atmospheric light
 				{
diff --git a/Engine/Source/Editor/Sequencer/Private/SSequencerSection.cpp b/Engine/Source/Editor/Sequencer/Private/SSequencerSection.cpp
--- a/Engine/Source/Editor/Sequencer/Private/SSequencerSection.cpp
+++ b/Engine/Source/Editor/Sequencer/Private/SSequencerSection.cpp
@@ -78,7 +78,7 @@ struct FSequencerSectionPainterImpl : FSequencerSectionPainter
 		FLinearColor FinalTint = FSequencerSectionPainter::BlendColor(Tint);
 		if (bIsHighlighted && Section.GetRange() != TRange<FFrameNumber>::All())
 		{
-			float Lum = FinalTint.ComputeLuminance() * 0.2f;
+			float Lum = FinalTint.GetLuminance() * 0.2f;
 			FinalTint = FinalTint + FLinearColor(Lum, Lum, Lum, 0.f);
 		}
 
diff --git a/Engine/Source/Editor/UnrealEd/Private/LevelEditorViewport.cpp b/Engine/Source/Editor/UnrealEd/Private/LevelEditorViewport.cpp
--- a/Engine/Source/Editor/UnrealEd/Private/LevelEditorViewport.cpp
+++ b/Engine/Source/Editor/UnrealEd/Private/LevelEditorViewport.cpp
@@ -167,7 +167,7 @@ static UDirectionalLightComponent* GetAtmosphericLight(const uint8 DesiredLightI
 			if (!AtmosphericLight->IsUsedAsAtmosphereSunLight() || AtmosphericLight->GetAtmosphereSunLightIndex() != DesiredLightIndex || !AtmosphericLight->GetVisibleFlag())
 				continue;
 
-			float LightLuminance = AtmosphericLight->GetColoredLightBrightness().ComputeLuminance();
+			float LightLuminance = AtmosphericLight->GetColoredLightBrightness().GetLuminance();
 			if (!SelectedAtmosphericLight ||					// Set it if null
 				SelectedLightLuminance < LightLuminance)		// Or choose the brightest atmospheric light
 			{
diff --git a/Engine/Source/Editor/UnrealEd/Private/Lightmass/Lightmass.cpp b/Engine/Source/Editor/UnrealEd/Private/Lightmass/Lightmass.cpp
--- a/Engine/Source/Editor/UnrealEd/Private/Lightmass/Lightmass.cpp
+++ b/Engine/Source/Editor/UnrealEd/Private/Lightmass/Lightmass.cpp
@@ -1046,10 +1046,10 @@ void FLightmassExporter::WriteLights( int32 Channel )
 		for (int32 LightIndex = 0; LightIndex < DirectionalLights.Num(); ++LightIndex)
 		{
 			const UDirectionalLightComponent* Light = DirectionalLights[LightIndex];
-			if (Light->IsUsedAsAtmosphereSunLight() && Light->GetColoredLightBrightness().ComputeLuminance() > AtmosphereLightsBrightness[Light->GetAtmosphereSunLightIndex()])
+			if (Light->IsUsedAsAtmosphereSunLight() && Light->GetColoredLightBrightness().GetLuminance() > AtmosphereLightsBrightness[Light->GetAtmosphereSunLightIndex()])
 			{
 				AtmosphereLights[Light->GetAtmosphereSunLightIndex()] = Light;
-				AtmosphereLightsBrightness[Light->GetAtmosphereSunLightIndex()] = Light->GetColoredLightBrightness().ComputeLuminance();
+				AtmosphereLightsBrightness[Light->GetAtmosphereSunLightIndex()] = Light->GetColoredLightBrightness().GetLuminance();
 			}
 		}
 
diff --git a/Engine/Source/Programs/UnrealLightmass/Private/Lighting/FinalGather.cpp b/Engine/Source/Programs/UnrealLightmass/Private/Lighting/FinalGather.cpp
--- a/Engine/Source/Programs/UnrealLightmass/Private/Lighting/FinalGather.cpp
+++ b/Engine/Source/Programs/UnrealLightmass/Private/Lighting/FinalGather.cpp
@@ -753,7 +753,7 @@ public:
 						{
 							const FLightingAndOcclusion RootElementLighting = GetRootValue(ThetaIndex, PhiIndex);
 							const FLinearColor Radiance = RootElementLighting.Lighting + RootElementLighting.StationarySkyLighting;
-							const float RelativeBrightness = Radiance.ComputeLuminance() / AverageBrightness;
+							const float RelativeBrightness = Radiance.GetLuminance() / AverageBrightness;
 
 							for (int32 NeighborIndex = 0; NeighborIndex < UE_ARRAY_COUNT(Neighbors); NeighborIndex++)
 							{
@@ -765,7 +765,7 @@ public:
 								if (NeighborTheta >= 0 && NeighborTheta < NumThetaSteps)
 								{
 									const FLightingAndOcclusion NeighborLighting = GetRootValue(NeighborTheta, NeighborPhi);
-									const float NeighborBrightness = (NeighborLighting.Lighting + NeighborLighting.StationarySkyLighting).ComputeLuminance();
+									const float NeighborBrightness = (NeighborLighting.Lighting + NeighborLighting.StationarySkyLighting).GetLuminance();
 									const float NeighborRelativeBrightness = NeighborBrightness / AverageBrightness;
 									MaxRelativeDifference = FMath::Max(MaxRelativeDifference, FMath::Abs(RelativeBrightness - NeighborRelativeBrightness));
 
@@ -891,7 +891,7 @@ public:
 							{
 								const FLightingAndOcclusion ChildLighting = ChildNode->Element.Lighting;
 								const FLinearColor Radiance = ChildLighting.Lighting + ChildLighting.StationarySkyLighting;
-								const float RelativeBrightness = Radiance.ComputeLuminance() / AverageBrightness;
+								const float RelativeBrightness = Radiance.GetLuminance() / AverageBrightness;
 
 								// Only search the axis neighbors past the first depth
 								for (int32 NeighborIndex = 0; NeighborIndex < UE_ARRAY_COUNT(Neighbors) / 2; NeighborIndex++)
@@ -902,7 +902,7 @@ public:
 									// Query must be done on the center of the cell
 									const FVector2D NeighborUV = FVector2D(NeighborU, NeighborV) + NodeContext.Size / 4;
 									const FLightingAndOcclusion NeighborLighting = GetValue(NeighborUV);
-									const float NeighborBrightness = (NeighborLighting.Lighting + NeighborLighting.StationarySkyLighting).ComputeLuminance();
+									const float NeighborBrightness = (NeighborLighting.Lighting + NeighborLighting.StationarySkyLighting).GetLuminance();
 									const float NeighborRelativeBrightness = NeighborBrightness / AverageBrightness;
 									MaxRelativeDifference = FMath::Max(MaxRelativeDifference, FMath::Abs(RelativeBrightness - NeighborRelativeBrightness));
 									MaxSkyOcclusionDifference = FMath::Max(MaxSkyOcclusionDifference, FMath::Abs(ChildLighting.UnoccludedSkyVector.SizeSquared() - NeighborLighting.UnoccludedSkyVector.SizeSquared()));
diff --git a/Engine/Source/Programs/UnrealLightmass/Private/Lighting/LightingMesh.cpp b/Engine/Source/Programs/UnrealLightmass/Private/Lighting/LightingMesh.cpp
--- a/Engine/Source/Programs/UnrealLightmass/Private/Lighting/LightingMesh.cpp
+++ b/Engine/Source/Programs/UnrealLightmass/Private/Lighting/LightingMesh.cpp
@@ -51,7 +51,7 @@ FLinearColor FStaticLightingMesh::SampleBRDF(
 
 	const FLinearColor Diffuse = EvaluateDiffuse(Vertex.TextureCoordinates[0], ElementIndex);
 
-	const float DiffuseIntensity = Diffuse.ComputeLuminance();
+	const float DiffuseIntensity = Diffuse.GetLuminance();
 	
 	// Generate a direction based on the cosine lobe
 	FVector4 TangentPathDirection = GetCosineHemisphereVector(RandomStream);
diff --git a/Engine/Source/Runtime/Core/Private/Math/Color.cpp b/Engine/Source/Runtime/Core/Private/Math/Color.cpp
--- a/Engine/Source/Runtime/Core/Private/Math/Color.cpp
+++ b/Engine/Source/Runtime/Core/Private/Math/Color.cpp
@@ -173,7 +173,7 @@ FColor FLinearColor::QuantizeRound() const
  */
 FLinearColor FLinearColor::Desaturate( float Desaturation ) const
 {
-	float Lum = ComputeLuminance();
+	float Lum = GetLuminance();
 	return FMath::Lerp( *this, FLinearColor( Lum, Lum, Lum, 0 ), Desaturation );
 }
 
diff --git a/Engine/Source/Runtime/Core/Public/Math/Color.h b/Engine/Source/Runtime/Core/Public/Math/Color.h
--- a/Engine/Source/Runtime/Core/Public/Math/Color.h
+++ b/Engine/Source/Runtime/Core/Public/Math/Color.h
@@ -339,18 +339,11 @@ struct FLinearColor
 	CORE_API FLinearColor Desaturate( float Desaturation ) const;
 
 	/** Computes the perceptually weighted luminance value of a color. */
-	inline float ComputeLuminance() const
+	inline float GetLuminance() const
 	{		
 		return R * 0.3f + G * 0.59f + B * 0.11f;
 	}
 	
-	FORCEINLINE float GetLuminance() const 
-	{
-		// duplicate function
-		// @todo remove me
-		return ComputeLuminance();
-	}
-
 	/**
 	 * Returns the maximum value in this color structure
 	 *
diff --git a/Engine/Source/Runtime/Engine/Private/Kismet/ImportanceSamplingLibrary.cpp b/Engine/Source/Runtime/Engine/Private/Kismet/ImportanceSamplingLibrary.cpp
--- a/Engine/Source/Runtime/Engine/Private/Kismet/ImportanceSamplingLibrary.cpp
+++ b/Engine/Source/Runtime/Engine/Private/Kismet/ImportanceSamplingLibrary.cpp
@@ -60,7 +60,7 @@ float FImportanceTexture::ImportanceWeight(FColor Texel, TEnumAsByte<EImportance
 	FLinearColor LinearTexel = (Texture.IsValid() && Texture->SRGB) ? FLinearColor(Texel) : Texel.ReinterpretAsLinear();
 
 	switch (WeightingFunc) {
-	case EImportanceWeight::Luminance:	return LinearTexel.ComputeLuminance();
+	case EImportanceWeight::Luminance:	return LinearTexel.GetLuminance();
 	case EImportanceWeight::Red:		return LinearTexel.R;
 	case EImportanceWeight::Green:		return LinearTexel.G;
 	case EImportanceWeight::Blue:		return LinearTexel.B;
diff --git a/Engine/Source/Runtime/Renderer/Private/FogRendering.cpp b/Engine/Source/Runtime/Renderer/Private/FogRendering.cpp
--- a/Engine/Source/Runtime/Renderer/Private/FogRendering.cpp
+++ b/Engine/Source/Runtime/Renderer/Private/FogRendering.cpp
@@ -350,7 +350,7 @@ void FSceneRenderer::InitFogConstants()
 				if (SunLight)
 				{
 					View.InscatteringLightDirection = -SunLight->Proxy->GetDirection();
-					View.DirectionalInscatteringColor = FogInfo.DirectionalInscatteringColor * SunLight->Proxy->GetColor().ComputeLuminance();
+					View.DirectionalInscatteringColor = FogInfo.DirectionalInscatteringColor * SunLight->Proxy->GetColor().GetLuminance();
 				}
 				View.bUseDirectionalInscattering = SunLight != nullptr;
 			}
diff --git a/Engine/Source/Runtime/Renderer/Private/RendererScene.cpp b/Engine/Source/Runtime/Renderer/Private/RendererScene.cpp
--- a/Engine/Source/Runtime/Renderer/Private/RendererScene.cpp
+++ b/Engine/Source/Runtime/Renderer/Private/RendererScene.cpp
@@ -3724,7 +3724,7 @@ void FScene::ProcessAtmosphereLightAddition_RenderThread(FLightSceneInfo* LightS
 	{
 		const uint8 Index = LightSceneInfo->Proxy->GetAtmosphereSunLightIndex();
 		if (!AtmosphereLights[Index] ||																								// Set it if null
-			LightSceneInfo->Proxy->GetColor().ComputeLuminance() > AtmosphereLights[Index]->Proxy->GetColor().ComputeLuminance())	// Or choose the brightest sun light
+			LightSceneInfo->Proxy->GetColor().GetLuminance() > AtmosphereLights[Index]->Proxy->GetColor().GetLuminance())	// Or choose the brightest sun light
 		{
 			AtmosphereLights[Index] = LightSceneInfo;
 		}
@@ -3743,7 +3743,7 @@ void FScene::ProcessAtmosphereLightRemoval_RenderThread(FLightSceneInfo* LightSc
 		for (TSparseArray<FLightSceneInfoCompact>::TConstIterator It(Lights); It; ++It)
 		{
 			const FLightSceneInfoCompact& LightInfo = *It;
-			float LightLuminance = LightInfo.LightSceneInfo->Proxy->GetColor().ComputeLuminance();
+			float LightLuminance = LightInfo.LightSceneInfo->Proxy->GetColor().GetLuminance();
 
 			if (LightInfo.LightSceneInfo != LightSceneInfo
 				&& LightInfo.LightSceneInfo->Proxy->IsUsedAsAtmosphereSunLight() && LightInfo.LightSceneInfo->Proxy->GetAtmosphereSunLightIndex() == Index
