Fixed UnrealLightmass compilation

diff --git a/Engine/Source/Programs/UnrealLightmass/Private/Lighting/LightingSystem.cpp b/Engine/Source/Programs/UnrealLightmass/Private/Lighting/LightingSystem.cpp
--- a/Engine/Source/Programs/UnrealLightmass/Private/Lighting/LightingSystem.cpp
+++ b/Engine/Source/Programs/UnrealLightmass/Private/Lighting/LightingSystem.cpp
@@ -913,7 +913,7 @@ void FStaticLightingSystem::ValidateSettings(FScene& InScene)
 		InScene.PhotonMappingSettings.bCacheIrradiancePhotonsOnSurfaces = false;
 	}
 	InScene.PhotonMappingSettings.FinalGatherImportanceSampleFraction = FMath::Clamp(InScene.PhotonMappingSettings.FinalGatherImportanceSampleFraction, 0.0f, 1.0f);
-	if (FMath::TruncToInt(InScene.ImportanceTracingSettings.NumHemisphereSamples * (1.0f - InScene.PhotonMappingSettings.FinalGatherImportanceSampleFraction) < 1))
+	if (InScene.ImportanceTracingSettings.NumHemisphereSamples * (1.0f - InScene.PhotonMappingSettings.FinalGatherImportanceSampleFraction) < 1)
 	{
 		// Irradiance caching needs some uniform samples
 		InScene.IrradianceCachingSettings.bAllowIrradianceCaching = false;
@@ -2117,7 +2117,7 @@ void FStaticLightingSystem::CalculateStaticShadowDepthMap(FGuid LightGuid)
 		{
 			const float AspectRatio = ShadowDepthMap->ShadowMapSizeX / (float)ShadowDepthMap->ShadowMapSizeY;
 			ShadowDepthMap->ShadowMapSizeY = FMath::TruncToInt(FMath::Sqrt(ShadowSettings.StaticShadowDepthMapMaxSamples / AspectRatio));
-			ShadowDepthMap->ShadowMapSizeX = FMath::TruncToInt(ShadowSettings.StaticShadowDepthMapMaxSamples / ShadowDepthMap->ShadowMapSizeY);
+			ShadowDepthMap->ShadowMapSizeX = FMath::TruncToInt((float)ShadowSettings.StaticShadowDepthMapMaxSamples / (float)ShadowDepthMap->ShadowMapSizeY);
 		}
 
 		// Allocate the shadow map
@@ -2209,7 +2209,7 @@ void FStaticLightingSystem::CalculateStaticShadowDepthMap(FGuid LightGuid)
 		{
 			const float AspectRatio = ShadowDepthMap->ShadowMapSizeX / (float)ShadowDepthMap->ShadowMapSizeY;
 			ShadowDepthMap->ShadowMapSizeY = FMath::TruncToInt(FMath::Sqrt(ShadowSettings.StaticShadowDepthMapMaxSamples / AspectRatio));
-			ShadowDepthMap->ShadowMapSizeX = FMath::TruncToInt(ShadowSettings.StaticShadowDepthMapMaxSamples / ShadowDepthMap->ShadowMapSizeY);
+			ShadowDepthMap->ShadowMapSizeX = FMath::TruncToInt((float)ShadowSettings.StaticShadowDepthMapMaxSamples / (float)ShadowDepthMap->ShadowMapSizeY);
 		}
 
 		ShadowDepthMap->ShadowMap.Empty(ShadowDepthMap->ShadowMapSizeX * ShadowDepthMap->ShadowMapSizeY);
@@ -2296,7 +2296,7 @@ void FStaticLightingSystem::CalculateStaticShadowDepthMap(FGuid LightGuid)
 		{
 			const float AspectRatio = ShadowDepthMap->ShadowMapSizeX / (float)ShadowDepthMap->ShadowMapSizeY;
 			ShadowDepthMap->ShadowMapSizeY = FMath::TruncToInt(FMath::Sqrt(ShadowSettings.StaticShadowDepthMapMaxSamples / AspectRatio));
-			ShadowDepthMap->ShadowMapSizeX = FMath::TruncToInt(ShadowSettings.StaticShadowDepthMapMaxSamples / ShadowDepthMap->ShadowMapSizeY);
+			ShadowDepthMap->ShadowMapSizeX = FMath::TruncToInt((float)ShadowSettings.StaticShadowDepthMapMaxSamples / (float)ShadowDepthMap->ShadowMapSizeY);
 		}
 
 		// Allocate the shadow map
diff --git a/Engine/Source/Programs/UnrealLightmass/Private/Lighting/TextureMapping.cpp b/Engine/Source/Programs/UnrealLightmass/Private/Lighting/TextureMapping.cpp
--- a/Engine/Source/Programs/UnrealLightmass/Private/Lighting/TextureMapping.cpp
+++ b/Engine/Source/Programs/UnrealLightmass/Private/Lighting/TextureMapping.cpp
@@ -118,8 +118,8 @@ void FFullStaticLightingVertex::ApplyVertexModifications(int32 ElementIndex, boo
 
 void FStaticLightingTextureMapping::Initialize(FStaticLightingSystem& System)
 {
-	SurfaceCacheSizeX = FMath::Max(FMath::TruncToInt(CachedSizeX / System.GeneralSettings.MappingSurfaceCacheDownsampleFactor), 6);
-	SurfaceCacheSizeY = FMath::Max(FMath::TruncToInt(CachedSizeY / System.GeneralSettings.MappingSurfaceCacheDownsampleFactor), 6);
+	SurfaceCacheSizeX = FMath::Max(FMath::TruncToInt((float)CachedSizeX / (float)System.GeneralSettings.MappingSurfaceCacheDownsampleFactor), 6);
+	SurfaceCacheSizeY = FMath::Max(FMath::TruncToInt((float)CachedSizeY / (float)System.GeneralSettings.MappingSurfaceCacheDownsampleFactor), 6);
 }
 
 /** Caches irradiance photons on a single texture mapping. */
@@ -2545,7 +2545,7 @@ void FStaticLightingSystem::CalculateDirectSignedDistanceFieldLightingTextureMap
 			{
 				const float AspectRatio = ShadowMapSizeX / (float)ShadowMapSizeY;
 				ShadowMapSizeY = FMath::TruncToInt(FMath::Sqrt(ShadowDepthMapMaxSamples / AspectRatio));
-				ShadowMapSizeX = FMath::TruncToInt(ShadowDepthMapMaxSamples / ShadowMapSizeY);
+				ShadowMapSizeX = ShadowDepthMapMaxSamples / ShadowMapSizeY;
 			}
 
 			TArray<float> ShadowMap;
