Fix CPU Lightmass by coverting LWC types to single precision

diff --git a/Engine/Source/Editor/UnrealEd/Private/Lightmass/Lightmass.cpp b/Engine/Source/Editor/UnrealEd/Private/Lightmass/Lightmass.cpp
--- a/Engine/Source/Editor/UnrealEd/Private/Lightmass/Lightmass.cpp
+++ b/Engine/Source/Editor/UnrealEd/Private/Lightmass/Lightmass.cpp
@@ -656,19 +656,19 @@ void FLightmassExporter::WriteToChannel( FLightmassStatistics& Stats, FGuid& Deb
 
 			for (int32 VolumeIndex = 0; VolumeIndex < ImportanceVolumes.Num(); VolumeIndex++)
 			{
-				FBox LMBox = ImportanceVolumes[VolumeIndex];
+				FBox3f LMBox(ImportanceVolumes[VolumeIndex]);
 				Swarm.WriteChannel(Channel, &LMBox, sizeof(LMBox));
 			}
 
 			for (int32 VolumeIndex = 0; VolumeIndex < CharacterIndirectDetailVolumes.Num(); VolumeIndex++)
 			{
-				FBox LMBox = CharacterIndirectDetailVolumes[VolumeIndex];
+				FBox3f LMBox(CharacterIndirectDetailVolumes[VolumeIndex]);
 				Swarm.WriteChannel(Channel, &LMBox, sizeof(LMBox));
 			}
 
 			for (int32 PortalIndex = 0; PortalIndex < Portals.Num(); PortalIndex++)
 			{
-				FMatrix Matrix = Portals[PortalIndex];
+				FMatrix44f Matrix(Portals[PortalIndex]);
 				Swarm.WriteChannel(Channel, &Matrix, sizeof(Matrix));
 			}
 
@@ -828,11 +828,16 @@ void FLightmassExporter::WriteVisibilityData( int32 Channel )
 		APrecomputedVisibilityVolume* Volume = *It;
 		if (World->ContainsActor(Volume) && IsValid(Volume))
 		{
-			FBox LMBox = Volume->GetComponentsBoundingBox(true);
+			FBox3f LMBox(Volume->GetComponentsBoundingBox(true));
 			Swarm.WriteChannel(Channel, &LMBox, sizeof(LMBox));
 
-			TArray<FPlane> Planes;
-			Volume->Brush->GetSurfacePlanes(Volume, Planes);
+			TArray<FPlane> Planes4d;
+			Volume->Brush->GetSurfacePlanes(Volume, Planes4d);			
+			TArray<FPlane4f> Planes;
+			for (FPlane Plane4d : Planes4d)
+			{
+				Planes.Add(FPlane4f(Plane4d));
+			}
 			const int32 NumPlanes = Planes.Num();
 			Swarm.WriteChannel( Channel, &NumPlanes, sizeof(NumPlanes) );
 			Swarm.WriteChannel( Channel, Planes.GetData(), Planes.Num() * Planes.GetTypeSize() );
@@ -854,7 +859,7 @@ void FLightmassExporter::WriteVisibilityData( int32 Channel )
 		APrecomputedVisibilityOverrideVolume* Volume = *It;
 		if (World->ContainsActor(Volume) && IsValid(Volume))
 		{
-			FBox LMBox = Volume->GetComponentsBoundingBox(true);
+			FBox3f LMBox(Volume->GetComponentsBoundingBox(true));
 			Swarm.WriteChannel(Channel, &LMBox, sizeof(LMBox));
 
 			TArray<int32> VisibilityIds;
@@ -939,7 +944,7 @@ void FLightmassExporter::WriteVolumetricLightmapData( int32 Channel )
 		AVolumetricLightmapDensityVolume* DetailVolume = VolumetricLightmapDensityVolumes[VolumeIndex];
 
 		Lightmass::FVolumetricLightmapDensityVolumeData VolumeData;
-		VolumeData.Bounds = DetailVolume->GetComponentsBoundingBox(true);
+		VolumeData.Bounds = FBox3f(DetailVolume->GetComponentsBoundingBox(true));
 		VolumeData.AllowedMipLevelRange = FIntPoint(DetailVolume->AllowedMipLevelRange.Min, DetailVolume->AllowedMipLevelRange.Max);
 
 		TArray<FPlane> Planes;
@@ -1555,7 +1560,7 @@ void FLightmassExporter::WriteBaseMeshInstanceData( int32 Channel, int32 MeshInd
 	MeshInstanceData.bCastShadowAsTwoSided = Mesh->Component->bCastShadowAsTwoSided;
 	MeshInstanceData.bMovable = (Mesh->Component->Mobility != EComponentMobility::Static);
 	MeshInstanceData.NumRelevantLights = Mesh->RelevantLights.Num();
-	MeshInstanceData.BoundingBox = Mesh->BoundingBox;
+	MeshInstanceData.BoundingBox = FBox3f(Mesh->BoundingBox);
 	Swarm.WriteChannel( Channel, &MeshInstanceData, sizeof(MeshInstanceData) );
 	const uint32 LightGuidsSize = Mesh->RelevantLights.Num() * sizeof(FGuid);
 	if( LightGuidsSize > 0 )
@@ -3474,9 +3479,9 @@ void FLightmassProcessor::ImportVolumeSamples()
 		const int32 Channel = Swarm.OpenChannel( *ChannelName, LM_VOLUMESAMPLES_CHANNEL_FLAGS );
 		if (Channel >= 0)
 		{
-			FVector4 UnusedVolumeCenter;
+			FVector4f UnusedVolumeCenter;
 			Swarm.ReadChannel(Channel, &UnusedVolumeCenter, sizeof(UnusedVolumeCenter));
-			FVector4 UnusedVolumeExtent;
+			FVector4f UnusedVolumeExtent;
 			Swarm.ReadChannel(Channel, &UnusedVolumeExtent, sizeof(UnusedVolumeExtent));
 
 			int32 NumStreamLevels = System.GetWorld()->GetStreamingLevels().Num();
@@ -3497,17 +3502,17 @@ void FLightmassProcessor::ImportVolumeSamples()
 					UMapBuildDataRegistry* CurrentRegistry = CurrentStorageLevel->GetOrCreateMapBuildData();
 					FPrecomputedLightVolumeData& CurrentLevelData = CurrentRegistry->AllocateLevelPrecomputedLightVolumeBuildData(CurrentLevel->LevelBuildDataId);
 
-					FBox LevelVolumeBounds(ForceInit);
+					FBox3f LevelVolumeBounds(ForceInit);
 
 					for (int32 SampleIndex = 0; SampleIndex < VolumeSamples.Num(); SampleIndex++)
 					{
 						const Lightmass::FVolumeLightingSampleData& CurrentSample = VolumeSamples[SampleIndex];
-						FVector SampleMin = CurrentSample.PositionAndRadius - FVector(CurrentSample.PositionAndRadius.W);
-						FVector SampleMax = CurrentSample.PositionAndRadius + FVector(CurrentSample.PositionAndRadius.W);
-						LevelVolumeBounds += FBox(SampleMin, SampleMax);
+						FVector4f SampleMin = CurrentSample.PositionAndRadius - FVector(CurrentSample.PositionAndRadius.W);
+						FVector4f SampleMax = CurrentSample.PositionAndRadius + FVector(CurrentSample.PositionAndRadius.W);
+						LevelVolumeBounds += FBox3f(SampleMin, SampleMax);
 					}
 
-					CurrentLevelData.Initialize(LevelVolumeBounds);
+					CurrentLevelData.Initialize(FBox(LevelVolumeBounds));
 
 					for (int32 SampleIndex = 0; SampleIndex < VolumeSamples.Num(); SampleIndex++)
 					{
@@ -3599,13 +3604,13 @@ void FLightmassProcessor::ImportPrecomputedVisibility()
 
 			for (int32 CellIndex = 0; CellIndex < NumCells; CellIndex++)
 			{
-				FBox Bounds;
+				FBox3f Bounds;
 				Swarm.ReadChannel(Channel, &Bounds, sizeof(Bounds));
 
 				// Use the same index for this task guid as it has in VisibilityBucketGuids, so that visibility cells are processed in a deterministic order
 				CompletedPrecomputedVisibilityCells[ArrayIndex].AddZeroed();
 				FUncompressedPrecomputedVisibilityCell& CurrentCell = CompletedPrecomputedVisibilityCells[ArrayIndex].Last();
-				CurrentCell.Bounds = Bounds;
+				CurrentCell.Bounds = FBox(Bounds);
 				ReadArray(Channel, CurrentCell.VisibilityData);
 			}
 
diff --git a/Engine/Source/Editor/UnrealEd/Private/StaticLightingSystem/StaticLightingDebug.cpp b/Engine/Source/Editor/UnrealEd/Private/StaticLightingSystem/StaticLightingDebug.cpp
--- a/Engine/Source/Editor/UnrealEd/Private/StaticLightingSystem/StaticLightingDebug.cpp
+++ b/Engine/Source/Editor/UnrealEd/Private/StaticLightingSystem/StaticLightingDebug.cpp
@@ -517,26 +517,26 @@ void DrawStaticLightingDebugInfo(const FSceneView* View,FPrimitiveDrawInterface*
 				{
 					if (GDebugStaticLightingInfo.bCornerValid[CornerIndex])
 					{
-						PDI->DrawPoint(GDebugStaticLightingInfo.TexelCorners[CornerIndex] + CurrentVertex.VertexNormal * .04f, FLinearColor(0, 1, 1), 4.0f, SDPG_World);
+						PDI->DrawPoint(FVector4(GDebugStaticLightingInfo.TexelCorners[CornerIndex] + CurrentVertex.VertexNormal * .04f), FLinearColor(0, 1, 1), 4.0f, SDPG_World);
 					}
 				}
-				PDI->DrawPoint(CurrentVertex.VertexPosition, NormalColor, 4.0f, SDPG_World);
-				DrawWireSphere(PDI, CurrentVertex.VertexPosition, NormalColor, GDebugStaticLightingInfo.SampleRadius, 36, SDPG_World);
+				PDI->DrawPoint(FVector4(CurrentVertex.VertexPosition), NormalColor, 4.0f, SDPG_World);
+				DrawWireSphere(PDI, FVector4(CurrentVertex.VertexPosition), NormalColor, GDebugStaticLightingInfo.SampleRadius, 36, SDPG_World);
 			}
-			PDI->DrawLine(CurrentVertex.VertexPosition, CurrentVertex.VertexPosition + CurrentVertex.VertexNormal * 10, NormalColor, SDPG_World);
+			PDI->DrawLine(FVector4(CurrentVertex.VertexPosition), FVector4(CurrentVertex.VertexPosition + CurrentVertex.VertexNormal * 10), NormalColor, SDPG_World);
 		}
 
 		for (int32 RayIndex = 0; RayIndex < GDebugStaticLightingInfo.ShadowRays.Num(); RayIndex++)
 		{
 			const FDebugStaticLightingRay& CurrentRay = GDebugStaticLightingInfo.ShadowRays[RayIndex];
-			PDI->DrawLine(CurrentRay.Start, CurrentRay.End, CurrentRay.bHit ? FColor::Red : FColor::Green, SDPG_World);
+			PDI->DrawLine(FVector4(CurrentRay.Start), FVector4(CurrentRay.End), CurrentRay.bHit ? FColor::Red : FColor::Green, SDPG_World);
 		}
 		
 		for (int32 RayIndex = 0; RayIndex < GDebugStaticLightingInfo.PathRays.Num(); RayIndex++)
 		{
 			const FDebugStaticLightingRay& CurrentRay = GDebugStaticLightingInfo.PathRays[RayIndex];
 			const FColor RayColor = CurrentRay.bHit ? (CurrentRay.bPositive ? FColor(255,255,150) : FColor(150,150,150)) : FColor(50,50,255);
-			PDI->DrawLine(CurrentRay.Start, CurrentRay.End, RayColor, SDPG_World);
+			PDI->DrawLine(FVector4(CurrentRay.Start), FVector4(CurrentRay.End), RayColor, SDPG_World);
 		}
 
 		for (int32 RecordIndex = 0; RecordIndex < GDebugStaticLightingInfo.CacheRecords.Num(); RecordIndex++)
@@ -544,91 +544,91 @@ void DrawStaticLightingDebugInfo(const FSceneView* View,FPrimitiveDrawInterface*
 			const FDebugLightingCacheRecord& CurrentRecord = GDebugStaticLightingInfo.CacheRecords[RecordIndex];
 			if (CurrentRecord.bNearSelectedTexel)
 			{
-				DrawWireSphere(PDI, CurrentRecord.Vertex.VertexPosition + CurrentRecord.Vertex.VertexNormal * .1f, CurrentRecord.bAffectsSelectedTexel ? FColor(50, 255, 100) : FColor(100, 100, 100), CurrentRecord.Radius, 36, SDPG_World);
-				PDI->DrawLine(CurrentRecord.Vertex.VertexPosition, CurrentRecord.Vertex.VertexPosition + CurrentRecord.Vertex.VertexNormal * 12, CurrentRecord.bAffectsSelectedTexel ? FColor(50, 255, 100) : FColor(100, 100, 100), SDPG_World);
+				DrawWireSphere(PDI, FVector4(CurrentRecord.Vertex.VertexPosition + CurrentRecord.Vertex.VertexNormal * .1f), CurrentRecord.bAffectsSelectedTexel ? FColor(50, 255, 100) : FColor(100, 100, 100), CurrentRecord.Radius, 36, SDPG_World);
+				PDI->DrawLine(FVector4(CurrentRecord.Vertex.VertexPosition), FVector4(CurrentRecord.Vertex.VertexPosition + CurrentRecord.Vertex.VertexNormal * 12), CurrentRecord.bAffectsSelectedTexel ? FColor(50, 255, 100) : FColor(100, 100, 100), SDPG_World);
 			}
-			PDI->DrawPoint(CurrentRecord.Vertex.VertexPosition + CurrentRecord.Vertex.VertexNormal * .1f, FLinearColor(.5, 1, .5), 2.0f, SDPG_World);
+			PDI->DrawPoint(FVector4(CurrentRecord.Vertex.VertexPosition + CurrentRecord.Vertex.VertexNormal * .1f), FLinearColor(.5, 1, .5), 2.0f, SDPG_World);
 		}
 
 		for (int32 PhotonIndex = 0; PhotonIndex < GDebugStaticLightingInfo.DirectPhotons.Num(); PhotonIndex++)
 		{
 			const FDebugPhoton& CurrentPhoton = GDebugStaticLightingInfo.DirectPhotons[PhotonIndex];
-			PDI->DrawLine(CurrentPhoton.Position, CurrentPhoton.Position + CurrentPhoton.Direction * 50, FColor(200, 200, 100), SDPG_World);
+			PDI->DrawLine(FVector4(CurrentPhoton.Position), FVector4(CurrentPhoton.Position + CurrentPhoton.Direction * 50), FColor(200, 200, 100), SDPG_World);
 		}
 
 		for (int32 PhotonIndex = 0; PhotonIndex < GDebugStaticLightingInfo.IndirectPhotons.Num(); PhotonIndex++)
 		{
 			const FDebugPhoton& CurrentPhoton = GDebugStaticLightingInfo.IndirectPhotons[PhotonIndex];
-			PDI->DrawLine(CurrentPhoton.Position, CurrentPhoton.Position + CurrentPhoton.Direction, FColor(200, 100, 100), SDPG_World);
+			PDI->DrawLine(FVector4(CurrentPhoton.Position), FVector4(CurrentPhoton.Position + CurrentPhoton.Direction), FColor(200, 100, 100), SDPG_World);
 		}
 
 		for (int32 PhotonIndex = 0; PhotonIndex < GDebugStaticLightingInfo.IrradiancePhotons.Num(); PhotonIndex++)
 		{
 			const FDebugPhoton& CurrentPhoton = GDebugStaticLightingInfo.IrradiancePhotons[PhotonIndex];
-			PDI->DrawLine(CurrentPhoton.Position, CurrentPhoton.Position + CurrentPhoton.Direction * 50, FColor(150, 100, 250), SDPG_World);
+			PDI->DrawLine(FVector4(CurrentPhoton.Position), FVector4(CurrentPhoton.Position + CurrentPhoton.Direction * 50), FColor(150, 100, 250), SDPG_World);
 		}
 
 		for (int32 PhotonIndex = 0; PhotonIndex < GDebugStaticLightingInfo.GatheredPhotons.Num(); PhotonIndex++)
 		{
 			const FDebugPhoton& CurrentPhoton = GDebugStaticLightingInfo.GatheredPhotons[PhotonIndex];
-			PDI->DrawLine(CurrentPhoton.Position, CurrentPhoton.Position + CurrentPhoton.Normal * 50, FColor(100, 100, 100), SDPG_World);
-			PDI->DrawLine(CurrentPhoton.Position, CurrentPhoton.Position + CurrentPhoton.Direction * 50, FColor(50, 255, 100), SDPG_World);
-			PDI->DrawPoint(CurrentPhoton.Position + CurrentPhoton.Direction * .1f, FLinearColor(.5, 1, .5), 4.0f, SDPG_World);
+			PDI->DrawLine(FVector4(CurrentPhoton.Position), FVector4(CurrentPhoton.Position + CurrentPhoton.Normal * 50), FColor(100, 100, 100), SDPG_World);
+			PDI->DrawLine(FVector4(CurrentPhoton.Position), FVector4(CurrentPhoton.Position + CurrentPhoton.Direction * 50), FColor(50, 255, 100), SDPG_World);
+			PDI->DrawPoint(FVector4(CurrentPhoton.Position + CurrentPhoton.Direction * .1f), FLinearColor(.5, 1, .5), 4.0f, SDPG_World);
 		}
 
 		for (int32 PhotonIndex = 0; PhotonIndex < GDebugStaticLightingInfo.GatheredImportancePhotons.Num(); PhotonIndex++)
 		{
 			const FDebugPhoton& CurrentPhoton = GDebugStaticLightingInfo.GatheredImportancePhotons[PhotonIndex];
-			PDI->DrawLine(CurrentPhoton.Position, CurrentPhoton.Position + CurrentPhoton.Normal * 50, FColor(100, 100, 100), SDPG_World);
-			PDI->DrawLine(CurrentPhoton.Position, CurrentPhoton.Position + CurrentPhoton.Direction * 50, FColor(200, 100, 100), SDPG_World);
-			PDI->DrawPoint(CurrentPhoton.Position + CurrentPhoton.Direction * .1f, FLinearColor(.5, 1, .5), 4.0f, SDPG_World);
+			PDI->DrawLine(FVector4(CurrentPhoton.Position), FVector4(CurrentPhoton.Position + CurrentPhoton.Normal * 50), FColor(100, 100, 100), SDPG_World);
+			PDI->DrawLine(FVector4(CurrentPhoton.Position), FVector4(CurrentPhoton.Position + CurrentPhoton.Direction * 50), FColor(200, 100, 100), SDPG_World);
+			PDI->DrawPoint(FVector4(CurrentPhoton.Position + CurrentPhoton.Direction * .1f), FLinearColor(.5, 1, .5), 4.0f, SDPG_World);
 		}
 		const FColor NodeColor(150, 170, 180);
 		for (int32 NodeIndex = 0; NodeIndex < GDebugStaticLightingInfo.GatheredPhotonNodes.Num(); NodeIndex++)
 		{
 			const FDebugOctreeNode& CurrentNode = GDebugStaticLightingInfo.GatheredPhotonNodes[NodeIndex];
-			PDI->DrawLine(CurrentNode.Center + FVector(CurrentNode.Extent.X, CurrentNode.Extent.Y, CurrentNode.Extent.Z), CurrentNode.Center + FVector(-CurrentNode.Extent.X, CurrentNode.Extent.Y, CurrentNode.Extent.Z), NodeColor, SDPG_World);
-			PDI->DrawLine(CurrentNode.Center + FVector(CurrentNode.Extent.X, CurrentNode.Extent.Y, CurrentNode.Extent.Z), CurrentNode.Center + FVector(CurrentNode.Extent.X, -CurrentNode.Extent.Y, CurrentNode.Extent.Z), NodeColor, SDPG_World);
-			PDI->DrawLine(CurrentNode.Center + FVector(CurrentNode.Extent.X, CurrentNode.Extent.Y, CurrentNode.Extent.Z), CurrentNode.Center + FVector(CurrentNode.Extent.X, CurrentNode.Extent.Y, -CurrentNode.Extent.Z), NodeColor, SDPG_World);
+			PDI->DrawLine(FVector4(CurrentNode.Center) + FVector(CurrentNode.Extent.X, CurrentNode.Extent.Y, CurrentNode.Extent.Z), FVector4(CurrentNode.Center) + FVector(-CurrentNode.Extent.X, CurrentNode.Extent.Y, CurrentNode.Extent.Z), NodeColor, SDPG_World);
+			PDI->DrawLine(FVector4(CurrentNode.Center) + FVector(CurrentNode.Extent.X, CurrentNode.Extent.Y, CurrentNode.Extent.Z), FVector4(CurrentNode.Center) + FVector(CurrentNode.Extent.X, -CurrentNode.Extent.Y, CurrentNode.Extent.Z), NodeColor, SDPG_World);
+			PDI->DrawLine(FVector4(CurrentNode.Center) + FVector(CurrentNode.Extent.X, CurrentNode.Extent.Y, CurrentNode.Extent.Z), FVector4(CurrentNode.Center) + FVector(CurrentNode.Extent.X, CurrentNode.Extent.Y, -CurrentNode.Extent.Z), NodeColor, SDPG_World);
 			
-			PDI->DrawLine(CurrentNode.Center - FVector(CurrentNode.Extent.X, CurrentNode.Extent.Y, CurrentNode.Extent.Z), CurrentNode.Center - FVector(-CurrentNode.Extent.X, CurrentNode.Extent.Y, CurrentNode.Extent.Z), NodeColor, SDPG_World);
-			PDI->DrawLine(CurrentNode.Center - FVector(CurrentNode.Extent.X, CurrentNode.Extent.Y, CurrentNode.Extent.Z), CurrentNode.Center - FVector(CurrentNode.Extent.X, -CurrentNode.Extent.Y, CurrentNode.Extent.Z), NodeColor, SDPG_World);
-			PDI->DrawLine(CurrentNode.Center - FVector(CurrentNode.Extent.X, CurrentNode.Extent.Y, CurrentNode.Extent.Z), CurrentNode.Center - FVector(CurrentNode.Extent.X, CurrentNode.Extent.Y, -CurrentNode.Extent.Z), NodeColor, SDPG_World);
+			PDI->DrawLine(FVector4(CurrentNode.Center) - FVector(CurrentNode.Extent.X, CurrentNode.Extent.Y, CurrentNode.Extent.Z), FVector4(CurrentNode.Center) - FVector(-CurrentNode.Extent.X, CurrentNode.Extent.Y, CurrentNode.Extent.Z), NodeColor, SDPG_World);
+			PDI->DrawLine(FVector4(CurrentNode.Center) - FVector(CurrentNode.Extent.X, CurrentNode.Extent.Y, CurrentNode.Extent.Z), FVector4(CurrentNode.Center) - FVector(CurrentNode.Extent.X, -CurrentNode.Extent.Y, CurrentNode.Extent.Z), NodeColor, SDPG_World);
+			PDI->DrawLine(FVector4(CurrentNode.Center) - FVector(CurrentNode.Extent.X, CurrentNode.Extent.Y, CurrentNode.Extent.Z), FVector4(CurrentNode.Center) - FVector(CurrentNode.Extent.X, CurrentNode.Extent.Y, -CurrentNode.Extent.Z), NodeColor, SDPG_World);
 			
-			PDI->DrawLine(CurrentNode.Center + FVector(CurrentNode.Extent.X, -CurrentNode.Extent.Y, CurrentNode.Extent.Z), CurrentNode.Center + FVector(CurrentNode.Extent.X, -CurrentNode.Extent.Y, -CurrentNode.Extent.Z), NodeColor, SDPG_World);
-			PDI->DrawLine(CurrentNode.Center + FVector(CurrentNode.Extent.X, -CurrentNode.Extent.Y, CurrentNode.Extent.Z), CurrentNode.Center + FVector(-CurrentNode.Extent.X, -CurrentNode.Extent.Y, CurrentNode.Extent.Z), NodeColor, SDPG_World);
+			PDI->DrawLine(FVector4(CurrentNode.Center) + FVector(CurrentNode.Extent.X, -CurrentNode.Extent.Y, CurrentNode.Extent.Z), FVector4(CurrentNode.Center) + FVector(CurrentNode.Extent.X, -CurrentNode.Extent.Y, -CurrentNode.Extent.Z), NodeColor, SDPG_World);
+			PDI->DrawLine(FVector4(CurrentNode.Center) + FVector(CurrentNode.Extent.X, -CurrentNode.Extent.Y, CurrentNode.Extent.Z), FVector4(CurrentNode.Center) + FVector(-CurrentNode.Extent.X, -CurrentNode.Extent.Y, CurrentNode.Extent.Z), NodeColor, SDPG_World);
 
-			PDI->DrawLine(CurrentNode.Center + FVector(-CurrentNode.Extent.X, CurrentNode.Extent.Y, CurrentNode.Extent.Z), CurrentNode.Center + FVector(-CurrentNode.Extent.X, -CurrentNode.Extent.Y, CurrentNode.Extent.Z), NodeColor, SDPG_World);
-			PDI->DrawLine(CurrentNode.Center + FVector(-CurrentNode.Extent.X, CurrentNode.Extent.Y, CurrentNode.Extent.Z), CurrentNode.Center + FVector(-CurrentNode.Extent.X, CurrentNode.Extent.Y, -CurrentNode.Extent.Z), NodeColor, SDPG_World);
+			PDI->DrawLine(FVector4(CurrentNode.Center) + FVector(-CurrentNode.Extent.X, CurrentNode.Extent.Y, CurrentNode.Extent.Z), FVector4(CurrentNode.Center) + FVector(-CurrentNode.Extent.X, -CurrentNode.Extent.Y, CurrentNode.Extent.Z), NodeColor, SDPG_World);
+			PDI->DrawLine(FVector4(CurrentNode.Center) + FVector(-CurrentNode.Extent.X, CurrentNode.Extent.Y, CurrentNode.Extent.Z), FVector4(CurrentNode.Center) + FVector(-CurrentNode.Extent.X, CurrentNode.Extent.Y, -CurrentNode.Extent.Z), NodeColor, SDPG_World);
 
-			PDI->DrawLine(CurrentNode.Center + FVector(CurrentNode.Extent.X, CurrentNode.Extent.Y, -CurrentNode.Extent.Z), CurrentNode.Center + FVector(CurrentNode.Extent.X, -CurrentNode.Extent.Y, -CurrentNode.Extent.Z), NodeColor, SDPG_World);
-			PDI->DrawLine(CurrentNode.Center + FVector(CurrentNode.Extent.X, CurrentNode.Extent.Y, -CurrentNode.Extent.Z), CurrentNode.Center + FVector(-CurrentNode.Extent.X, CurrentNode.Extent.Y, -CurrentNode.Extent.Z), NodeColor, SDPG_World);
+			PDI->DrawLine(FVector4(CurrentNode.Center) + FVector(CurrentNode.Extent.X, CurrentNode.Extent.Y, -CurrentNode.Extent.Z), FVector4(CurrentNode.Center) + FVector(CurrentNode.Extent.X, -CurrentNode.Extent.Y, -CurrentNode.Extent.Z), NodeColor, SDPG_World);
+			PDI->DrawLine(FVector4(CurrentNode.Center) + FVector(CurrentNode.Extent.X, CurrentNode.Extent.Y, -CurrentNode.Extent.Z), FVector4(CurrentNode.Center) + FVector(-CurrentNode.Extent.X, CurrentNode.Extent.Y, -CurrentNode.Extent.Z), NodeColor, SDPG_World);
 		}
 
 		if (GDebugStaticLightingInfo.bDirectPhotonValid)
 		{
 			const FDebugPhoton& DirectPhoton = GDebugStaticLightingInfo.GatheredDirectPhoton;
-			PDI->DrawLine(DirectPhoton.Position, DirectPhoton.Position + DirectPhoton.Direction * 60, FColor(255, 255, 100), SDPG_World);
-			PDI->DrawPoint(DirectPhoton.Position + DirectPhoton.Direction * .1f, FLinearColor(1, 1, .5), 4.0f, SDPG_World);
+			PDI->DrawLine(FVector4(DirectPhoton.Position), FVector4(DirectPhoton.Position + DirectPhoton.Direction * 60), FColor(255, 255, 100), SDPG_World);
+			PDI->DrawPoint(FVector4(DirectPhoton.Position + DirectPhoton.Direction * .1f), FLinearColor(1, 1, .5), 4.0f, SDPG_World);
 		}
 
 		for (int32 RayIndex = 0; RayIndex < GDebugStaticLightingInfo.IndirectPhotonPaths.Num(); RayIndex++)
 		{
 			const FDebugStaticLightingRay& CurrentRay = GDebugStaticLightingInfo.IndirectPhotonPaths[RayIndex];
-			PDI->DrawLine(CurrentRay.Start, CurrentRay.End, FColor::White, SDPG_World);
+			PDI->DrawLine(FVector4(CurrentRay.Start), FVector4(CurrentRay.End), FColor::White, SDPG_World);
 		}
 
 		for (int32 SampleIndex = 0; SampleIndex < GDebugStaticLightingInfo.VolumeLightingSamples.Num(); SampleIndex++)
 		{
 			const FDebugVolumeLightingSample& CurrentSample = GDebugStaticLightingInfo.VolumeLightingSamples[SampleIndex];
-			PDI->DrawPoint(CurrentSample.Position, CurrentSample.AverageIncidentRadiance * GEngine->LightingOnlyBrightness, 12.0f, SDPG_World);
+			PDI->DrawPoint(FVector4(CurrentSample.Position), CurrentSample.AverageIncidentRadiance * GEngine->LightingOnlyBrightness, 12.0f, SDPG_World);
 		}
 
 		for (int32 RayIndex = 0; RayIndex < GDebugStaticLightingInfo.PrecomputedVisibilityRays.Num(); RayIndex++)
 		{
 			const FDebugStaticLightingRay& CurrentRay = GDebugStaticLightingInfo.PrecomputedVisibilityRays[RayIndex];
 			const FColor RayColor = CurrentRay.bHit ? (CurrentRay.bPositive ? FColor(255,255,150) : FColor(150,150,150)) : FColor(50,50,255);
-			PDI->DrawLine(CurrentRay.Start, CurrentRay.End, RayColor, SDPG_World);
+			PDI->DrawLine(FVector4(CurrentRay.Start), FVector4(CurrentRay.End), RayColor, SDPG_World);
 		}
 	}
 }
@@ -644,7 +644,7 @@ void DrawStaticLightingDebugInfo(const FSceneView* View, FCanvas* Canvas)
 			if (CurrentRecord.bNearSelectedTexel)
 			{
 				FVector2D PixelLocation;
-				if(View->ScreenToPixel(View->WorldToScreen(CurrentRecord.Vertex.VertexPosition),PixelLocation))
+				if(View->ScreenToPixel(View->WorldToScreen(FVector4(CurrentRecord.Vertex.VertexPosition)),PixelLocation))
 				{
 					const FColor TagColor = CurrentRecord.bAffectsSelectedTexel ? FColor(50,160,200) : FColor(120,120,120);
 					Canvas->DrawShadowedString(PixelLocation.X,PixelLocation.Y, *FString::FromInt(CurrentRecord.RecordId), GEngine->GetSmallFont(), TagColor);
@@ -656,7 +656,7 @@ void DrawStaticLightingDebugInfo(const FSceneView* View, FCanvas* Canvas)
 		{
 			const FDebugPhoton& CurrentPhoton = GDebugStaticLightingInfo.GatheredImportancePhotons[PhotonIndex];
 			FVector2D PixelLocation;
-			if(View->ScreenToPixel(View->WorldToScreen(CurrentPhoton.Position),PixelLocation))
+			if(View->ScreenToPixel(View->WorldToScreen(FVector4(CurrentPhoton.Position)),PixelLocation))
 			{
 				const FColor TagColor = FColor(120,120,120);
 				Canvas->DrawShadowedString(PixelLocation.X,PixelLocation.Y, *FString::FromInt(CurrentPhoton.Id), GEngine->GetSmallFont(), TagColor);
@@ -669,7 +669,7 @@ void DrawStaticLightingDebugInfo(const FSceneView* View, FCanvas* Canvas)
 			if (CurrentRay.bHit && CurrentRay.bPositive)
 			{
 				FVector2D PixelLocation;
-				if(View->ScreenToPixel(View->WorldToScreen(CurrentRay.End),PixelLocation))
+				if(View->ScreenToPixel(View->WorldToScreen(FVector4(CurrentRay.End)),PixelLocation))
 				{
 					const FColor TagColor = FColor(180,180,120);
 					Canvas->DrawShadowedString(PixelLocation.X,PixelLocation.Y, *FString::FromInt(RayIndex), GEngine->GetSmallFont(), TagColor);
diff --git a/Engine/Source/Editor/UnrealEd/Private/StaticLightingSystem/StaticLightingPrivate.h b/Engine/Source/Editor/UnrealEd/Private/StaticLightingSystem/StaticLightingPrivate.h
--- a/Engine/Source/Editor/UnrealEd/Private/StaticLightingSystem/StaticLightingPrivate.h
+++ b/Engine/Source/Editor/UnrealEd/Private/StaticLightingSystem/StaticLightingPrivate.h
@@ -454,8 +454,8 @@ private:
 /** Stores debug information about a static lighting ray. */
 struct FDebugStaticLightingRay
 {
-	FVector4 Start;
-	FVector4 End;
+	FVector4f Start;
+	FVector4f End;
 	bool bHit;
 	bool bPositive;
 };
@@ -469,8 +469,8 @@ struct FDebugStaticLightingVertex
 		VertexPosition(InVertex.WorldPosition)
 	{}
 
-	FVector4 VertexNormal;
-	FVector4 VertexPosition;
+	FVector4f VertexNormal;
+	FVector4f VertexPosition;
 };
 
 struct FDebugLightingCacheRecord
@@ -485,20 +485,20 @@ struct FDebugLightingCacheRecord
 struct FDebugPhoton
 {
 	int32 Id;
-	FVector4 Position;
-	FVector4 Direction;
-	FVector4 Normal;
+	FVector4f Position;
+	FVector4f Direction;
+	FVector4f Normal;
 };
 
 struct FDebugOctreeNode
 {
-	FVector4 Center;
-	FVector4 Extent;
+	FVector4f Center;
+	FVector4f Extent;
 };
 
 struct FDebugVolumeLightingSample
 {
-	FVector4 Position;
+	FVector4f Position;
 	FLinearColor AverageIncidentRadiance;
 };
 
@@ -527,7 +527,7 @@ struct FDebugLightingOutput
 	TArray<FDebugStaticLightingRay> PrecomputedVisibilityRays;
 	bool bDirectPhotonValid;
 	FDebugPhoton GatheredDirectPhoton;
-	FVector4 TexelCorners[NumTexelCorners];
+	FVector4f TexelCorners[NumTexelCorners];
 	bool bCornerValid[NumTexelCorners];
 	float SampleRadius;
 
diff --git a/Engine/Source/Programs/UnrealLightmass/Private/ImportExport/Exporter.cpp b/Engine/Source/Programs/UnrealLightmass/Private/ImportExport/Exporter.cpp
--- a/Engine/Source/Programs/UnrealLightmass/Private/ImportExport/Exporter.cpp
+++ b/Engine/Source/Programs/UnrealLightmass/Private/ImportExport/Exporter.cpp
@@ -34,8 +34,8 @@ namespace Lightmass
 	void FLightmassSolverExporter::ExportVolumeLightingSamples(
 		bool bExportVolumeLightingDebugOutput,
 		const FVolumeLightingDebugOutput& DebugOutput,
-		const FVector4& VolumeCenter, 
-		const FVector4& VolumeExtent, 
+		const FVector4f& VolumeCenter, 
+		const FVector4f& VolumeExtent, 
 		const TMap<FGuid,TArray<FVolumeLightingSample> >& VolumeSamples) const
 	{
 		if (bExportVolumeLightingDebugOutput)
@@ -112,7 +112,7 @@ namespace Lightmass
 				FMeshAreaLightData LightData;
 				LightData.LevelGuid = CurrentLight.LevelGuid;
 
-				FVector4 AverageNormal(0,0,0);
+				FVector4f AverageNormal(0,0,0);
 				for (int32 PrimitiveIndex = 0; PrimitiveIndex < CurrentLight.Primitives.Num(); PrimitiveIndex++)
 				{
 					AverageNormal += CurrentLight.Primitives[PrimitiveIndex].SurfaceNormal * CurrentLight.Primitives[PrimitiveIndex].SurfaceArea;
@@ -123,7 +123,7 @@ namespace Lightmass
 				}
 				else
 				{
-					AverageNormal = FVector4(1,0,0);
+					AverageNormal = FVector4f(1,0,0);
 				}
 				// Offset the position somewhat to reduce the chance of the generated light being inside the mesh
 				LightData.Position = CurrentLight.Position + AverageNormal *MeshAreaLightGeneratedDynamicLightSurfaceOffset;
@@ -150,7 +150,7 @@ namespace Lightmass
 	}
 
 	/** Exports the volume distance field. */
-	void FLightmassSolverExporter::ExportVolumeDistanceField(int32 VolumeSizeX, int32 VolumeSizeY, int32 VolumeSizeZ, float VolumeMaxDistance, const FBox& DistanceFieldVolumeBounds, const TArray<FColor>& VolumeDistanceField) const
+	void FLightmassSolverExporter::ExportVolumeDistanceField(int32 VolumeSizeX, int32 VolumeSizeY, int32 VolumeSizeZ, float VolumeMaxDistance, const FBox3f& DistanceFieldVolumeBounds, const TArray<FColor>& VolumeDistanceField) const
 	{
 		const FString ChannelName = CreateChannelName(VolumeDistanceFieldGuid, LM_MESHAREALIGHTDATA_VERSION, LM_MESHAREALIGHTDATA_EXTENSION);
 		const int32 ErrorCode = Swarm->OpenChannel(*ChannelName, LM_MESHAREALIGHT_CHANNEL_FLAGS, true);
diff --git a/Engine/Source/Programs/UnrealLightmass/Private/ImportExport/Exporter.h b/Engine/Source/Programs/UnrealLightmass/Private/ImportExport/Exporter.h
--- a/Engine/Source/Programs/UnrealLightmass/Private/ImportExport/Exporter.h
+++ b/Engine/Source/Programs/UnrealLightmass/Private/ImportExport/Exporter.h
@@ -49,8 +49,8 @@ namespace Lightmass
 		void ExportVolumeLightingSamples(
 			bool bExportVolumeLightingDebugOutput,
 			const struct FVolumeLightingDebugOutput& DebugOutput,
-			const FVector4& VolumeCenter, 
-			const FVector4& VolumeExtent, 
+			const FVector4f& VolumeCenter, 
+			const FVector4f& VolumeExtent, 
 			const TMap<FGuid,TArray<class FVolumeLightingSample> >& VolumeSamples) const;
 
 		/** Exports dominant shadow information to Unreal. */
@@ -63,7 +63,7 @@ namespace Lightmass
 		void ExportMeshAreaLightData(const class TIndirectArray<FMeshAreaLight>& MeshAreaLights, float MeshAreaLightGeneratedDynamicLightSurfaceOffset) const;
 
 		/** Exports the volume distance field. */
-		void ExportVolumeDistanceField(int32 VolumeSizeX, int32 VolumeSizeY, int32 VolumeSizeZ, float VolumeMaxDistance, const FBox& DistanceFieldVolumeBounds, const TArray<FColor>& VolumeDistanceField) const;
+		void ExportVolumeDistanceField(int32 VolumeSizeX, int32 VolumeSizeY, int32 VolumeSizeZ, float VolumeMaxDistance, const FBox3f& DistanceFieldVolumeBounds, const TArray<FColor>& VolumeDistanceField) const;
 
 	private:
 		class FLightmassSwarm*	Swarm;
diff --git a/Engine/Source/Programs/UnrealLightmass/Private/ImportExport/LightmassScene.cpp b/Engine/Source/Programs/UnrealLightmass/Private/ImportExport/LightmassScene.cpp
--- a/Engine/Source/Programs/UnrealLightmass/Private/ImportExport/LightmassScene.cpp
+++ b/Engine/Source/Programs/UnrealLightmass/Private/ImportExport/LightmassScene.cpp
@@ -122,7 +122,7 @@ void FScene::Import( FLightmassImporter& Importer )
 	ImportanceBoundingBox.Init();
 	for (int32 VolumeIndex = 0; VolumeIndex < NumImportanceVolumes; VolumeIndex++)
 	{
-		FBox LMBox;
+		FBox3f LMBox;
 		Importer.ImportData(&LMBox);
 		ImportanceBoundingBox += LMBox;
 		ImportanceVolumes.Add(LMBox);
@@ -130,21 +130,21 @@ void FScene::Import( FLightmassImporter& Importer )
 
 	if (NumImportanceVolumes == 0)
 	{
-		ImportanceBoundingBox = FBox(FVector4(0,0,0), FVector4(0,0,0));
+		ImportanceBoundingBox = FBox3f(FVector4f(0,0,0), FVector4f(0,0,0));
 	}
 
 	for (int32 VolumeIndex = 0; VolumeIndex < NumCharacterIndirectDetailVolumes; VolumeIndex++)
 	{
-		FBox LMBox;
+		FBox3f LMBox;
 		Importer.ImportData(&LMBox);
 		CharacterIndirectDetailVolumes.Add(LMBox);
 	}
 
 	for (int32 PortalIndex = 0; PortalIndex < NumPortals; PortalIndex++)
 	{
-		FMatrix LMPortal;
+		FMatrix44f LMPortal;
 		Importer.ImportData(&LMPortal);
-		Portals.Add(FSphere(LMPortal.GetOrigin(), FVector2D(LMPortal.GetScaleVector().Y, LMPortal.GetScaleVector().Z).Size()));
+		Portals.Add(FSphere3f(LMPortal.GetOrigin(), FVector2f(LMPortal.GetScaleVector().Y, LMPortal.GetScaleVector().Z).Size()));
 	}
 
 	Importer.ImportArray(VisibilityBucketGuids, NumPrecomputedVisibilityBuckets);
@@ -188,9 +188,9 @@ void FScene::Import( FLightmassImporter& Importer )
 		Importer.ImportData(&LMVolumeData);
 		// LWC_TODO: SceneExport.h does #pragma pack(1) for non-mac/linux platforms meaning there's a 4 byte padding difference depending on platform when large world coordinates are enabled. Would be good to make this consistent.
 #if (!UE_LARGE_WORLD_COORDINATES_DISABLED && (PLATFORM_MAC || PLATFORM_LINUX))
-		static_assert(sizeof(LMVolumeData) == sizeof(FBox) + sizeof(FIntPoint) + sizeof(int32) + 4, "Update member copy");
+		static_assert(sizeof(LMVolumeData) == sizeof(FBox3f) + sizeof(FIntPoint) + sizeof(int32) + 4, "Update member copy");
 #else
-		static_assert(sizeof(LMVolumeData) == sizeof(FBox) + sizeof(FIntPoint) + sizeof(int32), "Update member copy");
+		static_assert(sizeof(LMVolumeData) == sizeof(FBox3f) + sizeof(FIntPoint) + sizeof(int32), "Update member copy");
 #endif
 
 		FVolumetricLightmapDensityVolume LMVolume;
@@ -304,9 +304,9 @@ void FScene::Import( FLightmassImporter& Importer )
 	}
 }
 
-FBoxSphereBounds FScene::GetImportanceBounds() const
+FBoxSphereBounds3f FScene::GetImportanceBounds() const
 {
-	const FBoxSphereBounds ImportanceBoundSphere(ImportanceBoundingBox);
+	const FBoxSphereBounds3f ImportanceBoundSphere(ImportanceBoundingBox);
 	return ImportanceBoundSphere;
 }
 
@@ -398,11 +398,11 @@ const FStaticLightingMapping* FScene::FindMappingByGuid(FGuid FindGuid) const
 }
 
 /** Returns true if the specified position is inside any of the importance volumes. */
-bool FScene::IsPointInImportanceVolume(const FVector4& Position, float Tolerance) const
+bool FScene::IsPointInImportanceVolume(const FVector4f& Position, float Tolerance) const
 {
 	for (int32 VolumeIndex = 0; VolumeIndex < ImportanceVolumes.Num(); VolumeIndex++)
 	{
-		FBox Volume = ImportanceVolumes[VolumeIndex];
+		FBox3f Volume = ImportanceVolumes[VolumeIndex];
 
 		if (Position.X + Tolerance > Volume.Min.X && Position.X - Tolerance < Volume.Max.X
 			&& Position.Y + Tolerance > Volume.Min.Y && Position.Y - Tolerance < Volume.Max.Y
@@ -414,11 +414,11 @@ bool FScene::IsPointInImportanceVolume(const FVector4& Position, float Tolerance
 	return false;
 }
 
-bool FScene::IsBoxInImportanceVolume(const FBox& QueryBox) const
+bool FScene::IsBoxInImportanceVolume(const FBox3f& QueryBox) const
 {
 	for (int32 VolumeIndex = 0; VolumeIndex < ImportanceVolumes.Num(); VolumeIndex++)
 	{
-		FBox Volume = ImportanceVolumes[VolumeIndex];
+		FBox3f Volume = ImportanceVolumes[VolumeIndex];
 
 		if (Volume.Intersect(QueryBox))
 		{
@@ -430,7 +430,7 @@ bool FScene::IsBoxInImportanceVolume(const FBox& QueryBox) const
 }
 
 /** Returns true if the specified position is inside any of the visibility volumes. */
-bool FScene::IsPointInVisibilityVolume(const FVector4& Position) const
+bool FScene::IsPointInVisibilityVolume(const FVector4f& Position) const
 {
 	for (int32 VolumeIndex = 0; VolumeIndex < PrecomputedVisibilityVolumes.Num(); VolumeIndex++)
 	{
@@ -438,7 +438,7 @@ bool FScene::IsPointInVisibilityVolume(const FVector4& Position) const
 		bool bInsideAllPlanes = true;
 		for (int32 PlaneIndex = 0; PlaneIndex < Volume.Planes.Num() && bInsideAllPlanes; PlaneIndex++)
 		{
-			const FPlane& Plane = Volume.Planes[PlaneIndex];
+			const FPlane4f& Plane = Volume.Planes[PlaneIndex];
 			bInsideAllPlanes = bInsideAllPlanes && Plane.PlaneDot(Position) < 0.0f;
 		}
 		if (bInsideAllPlanes)
@@ -449,7 +449,7 @@ bool FScene::IsPointInVisibilityVolume(const FVector4& Position) const
 	return false;
 }
 
-bool FScene::DoesBoxIntersectVisibilityVolume(const FBox& TestBounds) const
+bool FScene::DoesBoxIntersectVisibilityVolume(const FBox3f& TestBounds) const
 {
 	for (int32 VolumeIndex = 0; VolumeIndex < PrecomputedVisibilityVolumes.Num(); VolumeIndex++)
 	{
@@ -463,9 +463,9 @@ bool FScene::DoesBoxIntersectVisibilityVolume(const FBox& TestBounds) const
 }
 
 /** Returns accumulated bounds from all the visibility volumes. */
-FBox FScene::GetVisibilityVolumeBounds() const
+FBox3f FScene::GetVisibilityVolumeBounds() const
 {
-	FBox Bounds(ForceInit);
+	FBox3f Bounds(ForceInit);
 	for (int32 VolumeIndex = 0; VolumeIndex < PrecomputedVisibilityVolumes.Num(); VolumeIndex++)
 	{
 		const FPrecomputedVisibilityVolume& Volume = PrecomputedVisibilityVolumes[VolumeIndex];
@@ -473,7 +473,7 @@ FBox FScene::GetVisibilityVolumeBounds() const
 	}
 	if (PrecomputedVisibilityVolumes.Num() > 0)
 	{
-		FVector4 DoubleExtent = Bounds.GetExtent() * 2;
+		FVector4f DoubleExtent = Bounds.GetExtent() * 2;
 		DoubleExtent.X = DoubleExtent.X - FMath::Fmod(DoubleExtent.X, PrecomputedVisibilitySettings.CellSize) + PrecomputedVisibilitySettings.CellSize;
 		DoubleExtent.Y = DoubleExtent.Y - FMath::Fmod(DoubleExtent.Y, PrecomputedVisibilitySettings.CellSize) + PrecomputedVisibilitySettings.CellSize;
 		// Round the max up to the next cell boundary
@@ -482,11 +482,11 @@ FBox FScene::GetVisibilityVolumeBounds() const
 	}
 	else
 	{
-		return FBox(FVector4(0,0,0),FVector4(0,0,0));
+		return FBox3f(FVector4f(0,0,0),FVector4f(0,0,0));
 	}
 }
 
-bool FScene::GetVolumetricLightmapAllowedMipRange(const FVector4& Position, FIntPoint& OutRange) const
+bool FScene::GetVolumetricLightmapAllowedMipRange(const FVector4f& Position, FIntPoint& OutRange) const
 {
 	FIntPoint Range(INT_MAX, INT_MAX);
 
@@ -498,7 +498,7 @@ bool FScene::GetVolumetricLightmapAllowedMipRange(const FVector4& Position, FInt
 		bool bInsideAllPlanes = true;
 		for (int32 PlaneIndex = 0; PlaneIndex < Volume.Planes.Num() && bInsideAllPlanes; PlaneIndex++)
 		{
-			const FPlane& Plane = Volume.Planes[PlaneIndex];
+			const FPlane4f& Plane = Volume.Planes[PlaneIndex];
 			bInsideAllPlanes = bInsideAllPlanes && Plane.PlaneDot(Position) < 0.0f;
 		}
 		if (bInsideAllPlanes)
@@ -572,21 +572,21 @@ void FLight::Import( FLightmassImporter& Importer )
  * @param Bounds - The bounding volume to test.
  * @return True if the light affects the bounding volume
  */
-bool FLight::AffectsBounds(const FBoxSphereBounds& Bounds) const
+bool FLight::AffectsBounds(const FBoxSphereBounds3f& Bounds) const
 {
 	return true;
 }
 
-FSphere FLight::GetBoundingSphere() const
+FSphere3f FLight::GetBoundingSphere() const
 {
 	// Directional lights will have a radius of WORLD_MAX
-	return FSphere(FVector(0, 0, 0), (float)WORLD_MAX);
+	return FSphere3f(FVector3f(0, 0, 0), (float)WORLD_MAX);
 }
 
 /**
  * Computes the intensity of the direct lighting from this light on a specific point.
  */
-FLinearColor FLight::GetDirectIntensity(const FVector4& Point, bool bCalculateForIndirectLighting) const
+FLinearColor FLight::GetDirectIntensity(const FVector4f& Point, bool bCalculateForIndirectLighting) const
 {
 	// light profile (IES)
 	float LightProfileAttenuation;
@@ -644,9 +644,9 @@ void FDirectionalLight::Import( FLightmassImporter& Importer )
 }
 
 void FDirectionalLight::Initialize(
-	const FBoxSphereBounds& InSceneBounds, 
+	const FBoxSphereBounds3f& InSceneBounds, 
 	bool bInEmitPhotonsOutsideImportanceVolume,
-	const FBoxSphereBounds& InImportanceBounds, 
+	const FBoxSphereBounds3f& InImportanceBounds, 
 	float InIndirectDiskRadius, 
 	int32 InGridSize,
 	float InDirectPhotonDensity,
@@ -658,15 +658,15 @@ void FDirectionalLight::Initialize(
 	ImportanceBounds = InImportanceBounds;
 
 	// Vector through the scene bound's origin, along the direction of the light
-	const FVector4 SceneAxis = (SceneBounds.Origin + Direction * SceneBounds.SphereRadius) - (SceneBounds.Origin - Direction * SceneBounds.SphereRadius);
+	const FVector4f SceneAxis = (SceneBounds.Origin + Direction * SceneBounds.SphereRadius) - (SceneBounds.Origin - Direction * SceneBounds.SphereRadius);
 	const float SceneAxisLength = SceneBounds.SphereRadius * 2.0f;
-	const FVector4 DirectionalLightOriginToImportanceOrigin = ImportanceBounds.Origin - (SceneBounds.Origin - Direction * SceneBounds.SphereRadius);
+	const FVector4f DirectionalLightOriginToImportanceOrigin = ImportanceBounds.Origin - (SceneBounds.Origin - Direction * SceneBounds.SphereRadius);
 	// Find the closest point on the scene's axis to the importance volume's origin by projecting DirectionalLightOriginToImportanceOrigin onto SceneAxis.
 	// This gives the offset in the directional light's disk from the scene bound's origin.
-	const FVector4 ClosestPositionOnAxis = Dot3(SceneAxis, DirectionalLightOriginToImportanceOrigin) / (SceneAxisLength * SceneAxisLength) * SceneAxis + SceneBounds.Origin - Direction * SceneBounds.SphereRadius;
+	const FVector4f ClosestPositionOnAxis = Dot3(SceneAxis, DirectionalLightOriginToImportanceOrigin) / (SceneAxisLength * SceneAxisLength) * SceneAxis + SceneBounds.Origin - Direction * SceneBounds.SphereRadius;
 
 	// Find the disk offset from the world space origin and transform into the [-1,1] space of the directional light's disk, still in 3d.
-	const FVector4 DiskOffset = (ImportanceBounds.Origin - ClosestPositionOnAxis) / SceneBounds.SphereRadius;
+	const FVector4f DiskOffset = (ImportanceBounds.Origin - ClosestPositionOnAxis) / SceneBounds.SphereRadius;
 
 	const float DebugLength = (ImportanceBounds.Origin - ClosestPositionOnAxis).Size();
 	const float DebugDot = ((ImportanceBounds.Origin - ClosestPositionOnAxis) / DebugLength) | Direction;
@@ -674,14 +674,14 @@ void FDirectionalLight::Initialize(
 	//checkSlow(DebugLength < KINDA_SMALL_NUMBER * 10.0f || FMath::Abs(DebugDot) < DELTA * 10.0f);
 
 	// Decompose DiskOffset into it's corresponding parts along XAxis and YAxis
-	const FVector4 XAxisProjection = Dot3(XAxis, DiskOffset) * XAxis;
-	const FVector4 YAxisProjection = Dot3(YAxis, DiskOffset) * YAxis;
-	ImportanceDiskOrigin = FVector2D(Dot3(XAxisProjection, XAxis), Dot3(YAxisProjection, YAxis));
+	const FVector4f XAxisProjection = Dot3(XAxis, DiskOffset) * XAxis;
+	const FVector4f YAxisProjection = Dot3(YAxis, DiskOffset) * YAxis;
+	ImportanceDiskOrigin = FVector2f(Dot3(XAxisProjection, XAxis), Dot3(YAxisProjection, YAxis));
 
 	// Transform the importance volume's radius into the [-1,1] space of the directional light's disk
 	LightSpaceImportanceDiskRadius = ImportanceBounds.SphereRadius / SceneBounds.SphereRadius;
 
-	const FVector4 DebugPosition = (ImportanceDiskOrigin.X * XAxis + ImportanceDiskOrigin.Y * YAxis);
+	const FVector4f DebugPosition = (ImportanceDiskOrigin.X * XAxis + ImportanceDiskOrigin.Y * YAxis);
 	const float DebugLength2 = (DiskOffset - DebugPosition).Size3();
 	// Verify that DiskOffset was decomposed correctly by reconstructing it
 	checkSlow(DebugLength2 < KINDA_SMALL_NUMBER);
@@ -731,14 +731,14 @@ int32 FDirectionalLight::GetNumDirectPhotons(float DirectPhotonDensity) const
 }
 
 /** Generates a direction sample from the light's domain */
-void FDirectionalLight::SampleDirection(FLMRandomStream& RandomStream, FLightRay& SampleRay, FVector4& LightSourceNormal, FVector2D& LightSurfacePosition, float& RayPDF, FLinearColor& Power) const
+void FDirectionalLight::SampleDirection(FLMRandomStream& RandomStream, FLightRay& SampleRay, FVector4f& LightSourceNormal, FVector2f& LightSurfacePosition, float& RayPDF, FLinearColor& Power) const
 {
-	FVector4 DiskPosition3D;
+	FVector4f DiskPosition3D;
 	// If the importance volume is valid, generate samples in the importance volume with a probability of ImportanceBoundsSampleProbability
 	if (ImportanceBounds.SphereRadius > DELTA 
 		&& RandomStream.GetFraction() < ImportanceBoundsSampleProbability)
 	{
-		const FVector2D DiskPosition2D = GetUniformUnitDiskPosition(RandomStream);
+		const FVector2f DiskPosition2D = GetUniformUnitDiskPosition(RandomStream);
 		LightSurfacePosition = ImportanceDiskOrigin + DiskPosition2D * LightSpaceImportanceDiskRadius;
 		DiskPosition3D = SceneBounds.Origin + SceneBounds.SphereRadius * (LightSurfacePosition.X * XAxis + LightSurfacePosition.Y * YAxis);
 		RayPDF = ImportanceBoundsSampleProbability / ((float)PI * FMath::Square(ImportanceBounds.SphereRadius));
@@ -786,8 +786,8 @@ void FDirectionalLight::CachePathRays(const TArray<FIndirectPathRay>& IndirectPa
 
 	// Find the minimum and maximum position in the [-1, 1] space of the directional light's disk
 	// That a position can be generated from in FDirectionalLight::SampleDirection
-	FVector2D GridMin(1.0f, 1.0f);
-	FVector2D GridMax(-1.0f, -1.0f);
+	FVector2f GridMin(1.0f, 1.0f);
+	FVector2f GridMax(-1.0f, -1.0f);
 	for (int32 RayIndex = 0; RayIndex < IndirectPathRays.Num(); RayIndex++)
 	{
 		const FIndirectPathRay& CurrentRay = IndirectPathRays[RayIndex];
@@ -801,7 +801,7 @@ void FDirectionalLight::CachePathRays(const TArray<FIndirectPathRay>& IndirectPa
 	GridMax.X = FMath::Max(GridMax.X, -1.0f);
 	GridMax.Y = FMath::Max(GridMax.Y, -1.0f);
 	checkSlow(GridMax > GridMin);
-	const FVector2D GridExtent2D = 0.5f * (GridMax - GridMin);
+	const FVector2f GridExtent2D = 0.5f * (GridMax - GridMin);
 	// Keep the grid space square to simplify logic
 	GridExtent = FMath::Max(GridExtent2D.X, GridExtent2D.Y);
 	GridCenter = 0.5f * (GridMin + GridMax);
@@ -819,23 +819,23 @@ void FDirectionalLight::CachePathRays(const TArray<FIndirectPathRay>& IndirectPa
 		for (int32 X = 0; X < GridSize; X++)
 		{
 			// Center and Extent of the cell in the [0, 1] grid space
-			const FVector2D BoxCenter((X + .5f) * InvGridSize, (Y + .5f) * InvGridSize);
+			const FVector2f BoxCenter((X + .5f) * InvGridSize, (Y + .5f) * InvGridSize);
 			const float BoxExtent = .5f * InvGridSize;
 
 			// Corners of the cell
 			const int32 NumBoxCorners = 4;
-			FVector2D BoxCorners[NumBoxCorners];
-			BoxCorners[0] = BoxCenter + FVector2D(BoxExtent, BoxExtent);
-			BoxCorners[1] = BoxCenter + FVector2D(-BoxExtent, BoxExtent);
-			BoxCorners[2] = BoxCenter + FVector2D(BoxExtent, -BoxExtent);
-			BoxCorners[3] = BoxCenter + FVector2D(-BoxExtent, -BoxExtent);
+			FVector2f BoxCorners[NumBoxCorners];
+			BoxCorners[0] = BoxCenter + FVector2f(BoxExtent, BoxExtent);
+			BoxCorners[1] = BoxCenter + FVector2f(-BoxExtent, BoxExtent);
+			BoxCorners[2] = BoxCenter + FVector2f(BoxExtent, -BoxExtent);
+			BoxCorners[3] = BoxCenter + FVector2f(-BoxExtent, -BoxExtent);
 
 			// Calculate the world space positions of each corner of the cell
-			FVector4 WorldBoxCorners[NumBoxCorners];
+			FVector4f WorldBoxCorners[NumBoxCorners];
 			for (int32 i = 0; i < NumBoxCorners; i++)
 			{				
 				// Transform the cell corner from [0, 1] grid space to [-1, 1] in the directional light's disk
-				const FVector2D LightBoxCorner(2.0f * GridExtent * BoxCorners[i] + GridCenter - FVector2D(GridExtent, GridExtent));
+				const FVector2f LightBoxCorner(2.0f * GridExtent * BoxCorners[i] + GridCenter - FVector2f(GridExtent, GridExtent));
 				// Calculate the world position of the cell corner
 				WorldBoxCorners[i] = SceneBounds.Origin + SceneBounds.SphereRadius * (LightBoxCorner.X * XAxis + LightBoxCorner.Y * YAxis) - SceneBounds.SphereRadius * Direction;
 			}
@@ -882,12 +882,12 @@ void FDirectionalLight::SampleDirection(
 {
 	checkSlow(IndirectPathRays.Num() > 0);
 
-	const FVector2D DiskPosition2D = GetUniformUnitDiskPosition(RandomStream);
+	const FVector2f DiskPosition2D = GetUniformUnitDiskPosition(RandomStream);
 	const int32 RayIndex = FMath::TruncToInt(RandomStream.GetFraction() * IndirectPathRays.Num());
 	checkSlow(RayIndex >= 0 && RayIndex < IndirectPathRays.Num());
 
 	// Create the ray using a disk centered at the scene's origin, whose radius is the size of the scene
-	const FVector4 DiskPosition3D = IndirectPathRays[RayIndex].Start + IndirectDiskRadius * (DiskPosition2D.X * XAxis + DiskPosition2D.Y * YAxis);
+	const FVector4f DiskPosition3D = IndirectPathRays[RayIndex].Start + IndirectDiskRadius * (DiskPosition2D.X * XAxis + DiskPosition2D.Y * YAxis);
 	
 	SampleRay = FLightRay(
 		DiskPosition3D,
@@ -898,7 +898,7 @@ void FDirectionalLight::SampleDirection(
 
 	const float DiskPDF = 1.0f / ((float)PI * IndirectDiskRadius * IndirectDiskRadius);
 	const float LightSpaceIndirectDiskRadius = IndirectDiskRadius / SceneBounds.SphereRadius;
-	FVector2D SampleLightSurfacePosition;
+	FVector2f SampleLightSurfacePosition;
 	// Clamp the generated position to lie within the [-1, 1] space of the directional light's disk
 	SampleLightSurfacePosition.X = FMath::Clamp(DiskPosition2D.X * LightSpaceIndirectDiskRadius + IndirectPathRays[RayIndex].LightSurfacePosition.X, -1.0f, 1.0f - DELTA);
 	SampleLightSurfacePosition.Y = FMath::Clamp(DiskPosition2D.Y * LightSpaceIndirectDiskRadius + IndirectPathRays[RayIndex].LightSurfacePosition.Y, -1.0f, 1.0f - DELTA);
@@ -939,12 +939,12 @@ void FDirectionalLight::SampleDirection(
 float FDirectionalLight::Power() const
 {
 	const float EffectiveRadius = ImportanceBounds.SphereRadius > DELTA ? ImportanceBounds.SphereRadius : SceneBounds.SphereRadius;
-	const FLinearColor LightPower = GetDirectIntensity(FVector4(0,0,0), false) * IndirectLightingScale * (float)PI * EffectiveRadius * EffectiveRadius;
+	const FLinearColor LightPower = GetDirectIntensity(FVector4f(0,0,0), false) * IndirectLightingScale * (float)PI * EffectiveRadius * EffectiveRadius;
 	return FLinearColorUtils::LinearRGBToXYZ(LightPower).G;
 }
 
 /** Validates a surface sample given the position that sample is affecting. */
-void FDirectionalLight::ValidateSurfaceSample(const FVector4& Point, FLightSurfaceSample& Sample) const
+void FDirectionalLight::ValidateSurfaceSample(const FVector4f& Point, FLightSurfaceSample& Sample) const
 {
 	// Directional light samples are generated on a disk the size of the light source radius, centered on the origin
 	// Move the disk to the other side of the scene along the light's reverse direction
@@ -952,23 +952,23 @@ void FDirectionalLight::ValidateSurfaceSample(const FVector4& Point, FLightSurfa
 }
 
 /** Gets a single position which represents the center of the area light source from the ReceivingPosition's point of view. */
-FVector4 FDirectionalLight::LightCenterPosition(const FVector4& ReceivingPosition, const FVector4& ReceivingNormal) const
+FVector4f FDirectionalLight::LightCenterPosition(const FVector4f& ReceivingPosition, const FVector4f& ReceivingNormal) const
 {
 	return ReceivingPosition - Direction * 2.0f * SceneBounds.SphereRadius;
 }
 
 /** Returns true if all parts of the light are behind the surface being tested. */
-bool FDirectionalLight::BehindSurface(const FVector4& TrianglePoint, const FVector4& TriangleNormal) const
+bool FDirectionalLight::BehindSurface(const FVector4f& TrianglePoint, const FVector4f& TriangleNormal) const
 {
 	const float NormalDotLight = Dot3(TriangleNormal, FDirectionalLight::GetDirectLightingDirection(TrianglePoint, TriangleNormal));
 	return NormalDotLight < 0.0f;
 }
 
 /** Gets a single direction to use for direct lighting that is representative of the whole area light. */
-FVector4 FDirectionalLight::GetDirectLightingDirection(const FVector4& Point, const FVector4& PointNormal) const
+FVector4f FDirectionalLight::GetDirectLightingDirection(const FVector4f& Point, const FVector4f& PointNormal) const
 {
 	// The position on the directional light surface disk that will first be visible to a triangle rotating toward the light
-	const FVector4 FirstVisibleLightPoint = Point - Direction * 2.0f * SceneBounds.SphereRadius + PointNormal * LightSourceRadius;
+	const FVector4f FirstVisibleLightPoint = Point - Direction * 2.0f * SceneBounds.SphereRadius + PointNormal * LightSourceRadius;
 	return FirstVisibleLightPoint - Point;
 }
 
@@ -1013,7 +1013,7 @@ int32 FPointLight::GetNumDirectPhotons(float DirectPhotonDensity) const
  * @param Bounds - The bounding volume to test.
  * @return True if the light affects the bounding volume
  */
-bool FPointLight::AffectsBounds(const FBoxSphereBounds& Bounds) const
+bool FPointLight::AffectsBounds(const FBoxSphereBounds3f& Bounds) const
 {
 	if((Bounds.Origin - Position).SizeSquared() > FMath::Square(Radius + Bounds.SphereRadius))
 	{
@@ -1031,20 +1031,20 @@ bool FPointLight::AffectsBounds(const FBoxSphereBounds& Bounds) const
 /**
  * Computes the intensity of the direct lighting from this light on a specific point.
  */
-FLinearColor FPointLight::GetDirectIntensity(const FVector4& Point, bool bCalculateForIndirectLighting) const
+FLinearColor FPointLight::GetDirectIntensity(const FVector4f& Point, bool bCalculateForIndirectLighting) const
 {
 	if (LightFlags & GI_LIGHT_INVERSE_SQUARED)
 	{
-		FVector4 ToLight = Position - Point;
+		FVector4f ToLight = Position - Point;
 		float DistanceSqr = ToLight.SizeSquared3();
 
 		float DistanceAttenuation = 0.0f;
 		if( LightSourceLength > 0.0f )
 		{
 			// Line segment irradiance
-			FVector4 L01 =  GetLightTangent() * LightSourceLength;
-			FVector4 L0 = ToLight - 0.5f * L01;
-			FVector4 L1 = ToLight + 0.5f * L01;
+			FVector4f L01 =  GetLightTangent() * LightSourceLength;
+			FVector4f L0 = ToLight - 0.5f * L01;
+			FVector4f L1 = ToLight + 0.5f * L01;
 			float LengthL0 = L0.Size3();
 			float LengthL1 = L1.Size3();
 
@@ -1064,14 +1064,14 @@ FLinearColor FPointLight::GetDirectIntensity(const FVector4& Point, bool bCalcul
 	}
 	else
 	{
-		float RadialAttenuation = FMath::Pow(FMath::Max(1.0f - ((Position - Point) / Radius).SizeSquared3(), 0.0f), (FVector4::FReal)FalloffExponent);
+		float RadialAttenuation = FMath::Pow(FMath::Max(1.0f - ((Position - Point) / Radius).SizeSquared3(), 0.0f), (FVector4f::FReal)FalloffExponent);
 
 		return FLight::GetDirectIntensity(Point, bCalculateForIndirectLighting) * RadialAttenuation;
 	}
 }
 
 /** Returns an intensity scale based on the receiving point. */
-float FPointLight::CustomAttenuation(const FVector4& Point, FLMRandomStream& RandomStream, bool bMaintainEvenDensity) const
+float FPointLight::CustomAttenuation(const FVector4f& Point, FLMRandomStream& RandomStream, bool bMaintainEvenDensity) const
 {
 	// Remove the physical attenuation, then attenuation using Unreal point light radial falloff
 	const float PointDistanceSquared = (Position - Point).SizeSquared3();
@@ -1085,7 +1085,7 @@ float FPointLight::CustomAttenuation(const FVector4& Point, FLMRandomStream& Ran
 	}
 	else
 	{
-		UnrealAttenuation = FMath::Pow(FMath::Max(1.0f - ((Position - Point) / Radius).SizeSquared3(), 0.0f), (FVector4::FReal)FalloffExponent);
+		UnrealAttenuation = FMath::Pow(FMath::Max(1.0f - ((Position - Point) / Radius).SizeSquared3(), 0.0f), (FVector4f::FReal)FalloffExponent);
 	}
 
 	// light profile (IES)
@@ -1118,9 +1118,9 @@ float FPointLight::CustomAttenuation(const FVector4& Point, FLMRandomStream& Ran
 static const float PointLightIntensityScale = 1.0f; 
 
 /** Generates a direction sample from the light's domain */
-void FPointLight::SampleDirection(FLMRandomStream& RandomStream, FLightRay& SampleRay, FVector4& LightSourceNormal, FVector2D& LightSurfacePosition, float& RayPDF, FLinearColor& Power) const
+void FPointLight::SampleDirection(FLMRandomStream& RandomStream, FLightRay& SampleRay, FVector4f& LightSourceNormal, FVector2f& LightSurfacePosition, float& RayPDF, FLinearColor& Power) const
 {
-	const FVector4 RandomDirection = GetUnitVector(RandomStream);
+	const FVector4f RandomDirection = GetUnitVector(RandomStream);
 
 	FLightSurfaceSample SurfaceSample;
 	SampleLightSurface(RandomStream, SurfaceSample);
@@ -1129,7 +1129,7 @@ void FPointLight::SampleDirection(FLMRandomStream& RandomStream, FLightRay& Samp
 	if (SurfacePositionDotDirection < 0.0f)
 	{
 		// Reflect the surface position about the origin so that it lies in the same hemisphere as the RandomDirection
-		const FVector4 LocalSamplePosition = SurfaceSample.Position - Position;
+		const FVector4f LocalSamplePosition = SurfaceSample.Position - Position;
 		SurfaceSample.Position = -LocalSamplePosition + Position;
 	}
 
@@ -1161,14 +1161,14 @@ void FPointLight::SampleDirection(
 	const int32 RayIndex = FMath::TruncToInt(RandomStream.GetFraction() * IndirectPathRays.Num());
 	checkSlow(RayIndex >= 0 && RayIndex < IndirectPathRays.Num());
 
-	const FVector4 PathRayDirection = IndirectPathRays[RayIndex].UnitDirection;
+	const FVector4f PathRayDirection = IndirectPathRays[RayIndex].UnitDirection;
 
-	FVector4 XAxis(0,0,0);
-	FVector4 YAxis(0,0,0);
+	FVector4f XAxis(0,0,0);
+	FVector4f YAxis(0,0,0);
 	GenerateCoordinateSystem(PathRayDirection, XAxis, YAxis);
 
 	// Generate a sample direction within a cone about the indirect path
-	const FVector4 ConeSampleDirection = UniformSampleCone(RandomStream, CosIndirectPhotonEmitConeAngle, XAxis, YAxis, PathRayDirection);
+	const FVector4f ConeSampleDirection = UniformSampleCone(RandomStream, CosIndirectPhotonEmitConeAngle, XAxis, YAxis, PathRayDirection);
 
 	FLightSurfaceSample SurfaceSample;
 	// Generate a surface sample, not taking the indirect path into account
@@ -1178,7 +1178,7 @@ void FPointLight::SampleDirection(
 	if (SurfacePositionDotDirection < 0.0f)
 	{
 		// Reflect the surface position about the origin so that it lies in the same hemisphere as the ConeSampleDirection
-		const FVector4 LocalSamplePosition = SurfaceSample.Position - Position;
+		const FVector4f LocalSamplePosition = SurfaceSample.Position - Position;
 		SurfaceSample.Position = -LocalSamplePosition + Position;
 	}
 
@@ -1206,17 +1206,17 @@ void FPointLight::SampleDirection(
 }
 
 /** Validates a surface sample given the position that sample is affecting. */
-void FPointLight::ValidateSurfaceSample(const FVector4& Point, FLightSurfaceSample& Sample) const
+void FPointLight::ValidateSurfaceSample(const FVector4f& Point, FLightSurfaceSample& Sample) const
 {
 	// Only attempt to fixup sphere light source sample positions as the light source is radially symmetric
 	if (LightSourceLength <= 0)
 	{
-		const FVector4 LightToPoint = Point - Position;
+		const FVector4f LightToPoint = Point - Position;
 		const float LightToPointDistanceSquared = LightToPoint.SizeSquared3();
 		if (LightToPointDistanceSquared < FMath::Square(LightSourceRadius * 2.0f))
 		{
 			// Point is inside the light source radius * 2
-			FVector4 LocalSamplePosition = Sample.Position - Position;
+			FVector4f LocalSamplePosition = Sample.Position - Position;
 			// Reposition the light surface sample on a sphere whose radius is half of the distance from the light to Point
 			LocalSamplePosition *= FMath::Sqrt(LightToPointDistanceSquared) / (2.0f * LightSourceRadius);
 			Sample.Position = LocalSamplePosition + Position;
@@ -1227,7 +1227,7 @@ void FPointLight::ValidateSurfaceSample(const FVector4& Point, FLightSurfaceSamp
 		{
 			// Reflect the surface position about the origin so that it lies in the hemisphere facing Point
 			// The sample's PDF is unchanged
-			const FVector4 LocalSamplePosition = Sample.Position - Position;
+			const FVector4f LocalSamplePosition = Sample.Position - Position;
 			Sample.Position = -LocalSamplePosition + Position;
 		}
 	}
@@ -1256,13 +1256,13 @@ float FPointLight::Power() const
 	return FLinearColorUtils::LinearRGBToXYZ(LightPower).G;
 }
 
-FVector4 FPointLight::LightCenterPosition(const FVector4& ReceivingPosition, const FVector4& ReceivingNormal) const
+FVector4f FPointLight::LightCenterPosition(const FVector4f& ReceivingPosition, const FVector4f& ReceivingNormal) const
 {
 	if( LightSourceLength > 0 )
 	{
-		FVector4 ToLight = Position - ReceivingPosition;
+		FVector4f ToLight = Position - ReceivingPosition;
 
-		FVector4 Dir = GetLightTangent();
+		FVector4f Dir = GetLightTangent();
 		if( Dot3( ReceivingNormal, Dir ) < 0.0f )
 		{
 			Dir = -Dir;
@@ -1281,23 +1281,23 @@ FVector4 FPointLight::LightCenterPosition(const FVector4& ReceivingPosition, con
 }
 
 /** Returns true if all parts of the light are behind the surface being tested. */
-bool FPointLight::BehindSurface(const FVector4& TrianglePoint, const FVector4& TriangleNormal) const
+bool FPointLight::BehindSurface(const FVector4f& TrianglePoint, const FVector4f& TriangleNormal) const
 {
 	const float NormalDotLight = Dot3(TriangleNormal, FPointLight::GetDirectLightingDirection(TrianglePoint, TriangleNormal));
 	return NormalDotLight < 0.0f;
 }
 
 /** Gets a single direction to use for direct lighting that is representative of the whole area light. */
-FVector4 FPointLight::GetDirectLightingDirection(const FVector4& Point, const FVector4& PointNormal) const
+FVector4f FPointLight::GetDirectLightingDirection(const FVector4f& Point, const FVector4f& PointNormal) const
 {
-	FVector4 LightPosition = Position;
+	FVector4f LightPosition = Position;
 
 	if( LightSourceLength > 0 )
 	{
-		FVector4 ToLight = Position - Point;
-		FVector4 L01 =  GetLightTangent() * LightSourceLength;
-		FVector4 L0 = ToLight - 0.5 * L01;
-		FVector4 L1 = ToLight + 0.5 * L01;
+		FVector4f ToLight = Position - Point;
+		FVector4f L01 =  GetLightTangent() * LightSourceLength;
+		FVector4f L0 = ToLight - 0.5 * L01;
+		FVector4f L1 = ToLight + 0.5 * L01;
 #if 0
 		// Point on line segment with smallest angle to normal
 		float A = LightSourceLength * LightSourceLength;
@@ -1317,12 +1317,12 @@ FVector4 FPointLight::GetDirectLightingDirection(const FVector4& Point, const FV
 	else
 	{
 		// The position on the point light surface sphere that will first be visible to a triangle rotating toward the light
-		const FVector4 FirstVisibleLightPoint = LightPosition + PointNormal * LightSourceRadius;
+		const FVector4f FirstVisibleLightPoint = LightPosition + PointNormal * LightSourceRadius;
 		return FirstVisibleLightPoint - Point;
 	}
 }
 
-FVector FPointLight::GetLightTangent() const
+FVector3f FPointLight::GetLightTangent() const
 {
 	return LightTangent;
 }
@@ -1330,13 +1330,13 @@ FVector FPointLight::GetLightTangent() const
 /** Generates a sample on the light's surface. */
 void FPointLight::SampleLightSurface(FLMRandomStream& RandomStream, FLightSurfaceSample& Sample) const
 {
-	Sample.DiskPosition = FVector2D(0, 0);
+	Sample.DiskPosition = FVector2f(0, 0);
 
 	if (LightSourceLength <= 0)
 	{
 		// Generate a sample on the surface of the sphere with uniform density over the surface area of the sphere
 		//@todo - stratify
-		const FVector4 UnitSpherePosition = GetUnitVector(RandomStream);
+		const FVector4f UnitSpherePosition = GetUnitVector(RandomStream);
 		Sample.Position = UnitSpherePosition * LightSourceRadius + Position;
 		Sample.Normal = UnitSpherePosition;
 		// Probability of generating this surface position is 1 / SurfaceArea
@@ -1349,7 +1349,7 @@ void FPointLight::SampleLightSurface(FLMRandomStream& RandomStream, FLightSurfac
 		float SphereSurfaceArea = 4.0f * (float)PI * ClampedLightSourceRadius * ClampedLightSourceRadius;
 		float TotalSurfaceArea = CylinderSurfaceArea + SphereSurfaceArea;
 
-		const FVector TubeLightDirection = GetLightTangent();
+		const FVector3f TubeLightDirection = GetLightTangent();
 
 		// Cylinder End caps
 		// The chance of calculating a point on the end sphere is equal to it's percentage of total surface area
@@ -1357,7 +1357,7 @@ void FPointLight::SampleLightSurface(FLMRandomStream& RandomStream, FLightSurfac
 		{
 			// Generate a sample on the surface of the sphere with uniform density over the surface area of the sphere
 			//@todo - stratify
-			const FVector4 UnitSpherePosition = GetUnitVector(RandomStream);
+			const FVector4f UnitSpherePosition = GetUnitVector(RandomStream);
 			Sample.Position = UnitSpherePosition * ClampedLightSourceRadius + Position;
 
 			if (Dot3(UnitSpherePosition, TubeLightDirection) > 0)
@@ -1375,11 +1375,11 @@ void FPointLight::SampleLightSurface(FLMRandomStream& RandomStream, FLightSurfac
 		else
 		{
 			// Get point along center line
-			FVector4 CentreLinePosition = Position + TubeLightDirection * LightSourceLength * (RandomStream.GetFraction() - 0.5f);
+			FVector4f CentreLinePosition = Position + TubeLightDirection * LightSourceLength * (RandomStream.GetFraction() - 0.5f);
 			// Get point radius away from center line at random angle
 			float Theta = 2.0f * (float)PI * RandomStream.GetFraction();
-			FVector4 CylEdgePos = FVector4(0, FMath::Cos(Theta), FMath::Sin(Theta), 1);
-			CylEdgePos = FRotationMatrix::MakeFromZ( TubeLightDirection ).TransformVector( CylEdgePos );
+			FVector4f CylEdgePos = FVector4f(0, FMath::Cos(Theta), FMath::Sin(Theta), 1);
+			CylEdgePos = FRotationMatrix44f::MakeFromZ( TubeLightDirection ).TransformVector( CylEdgePos );
 
 			Sample.Position = CylEdgePos * ClampedLightSourceRadius + CentreLinePosition;
 			Sample.Normal = CylEdgePos;
@@ -1417,7 +1417,7 @@ void FSpotLight::Initialize(float InIndirectPhotonEmitConeAngle)
  * @param Bounds - The bounding volume to test.
  * @return True if the light affects the bounding volume
  */
-bool FSpotLight::AffectsBounds(const FBoxSphereBounds& Bounds) const
+bool FSpotLight::AffectsBounds(const FBoxSphereBounds3f& Bounds) const
 {
 	if(!FLight::AffectsBounds(Bounds))
 	{
@@ -1433,7 +1433,7 @@ bool FSpotLight::AffectsBounds(const FBoxSphereBounds& Bounds) const
 	// Cone check
 
 
-	FVector4	U = Position - (Bounds.SphereRadius / SinOuterConeAngle) * Direction,
+	FVector4f	U = Position - (Bounds.SphereRadius / SinOuterConeAngle) * Direction,
 				D = Bounds.Origin - U;
 	float	dsqr = Dot3(D, D),
 			E = Dot3(Direction, D);
@@ -1451,31 +1451,31 @@ bool FSpotLight::AffectsBounds(const FBoxSphereBounds& Bounds) const
 	return false;
 }
 
-FSphere FSpotLight::GetBoundingSphere() const
+FSphere3f FSpotLight::GetBoundingSphere() const
 {
-	return (FSphere)FMath::ComputeBoundingSphereForCone((FVector3f)Position, (FVector3f)Direction, Radius, CosOuterConeAngle, SinOuterConeAngle);
+	return (FSphere3f)FMath::ComputeBoundingSphereForCone((FVector3f)Position, (FVector3f)Direction, Radius, CosOuterConeAngle, SinOuterConeAngle);
 }
 
 /**
  * Computes the intensity of the direct lighting from this light on a specific point.
  */
-FLinearColor FSpotLight::GetDirectIntensity(const FVector4& Point, bool bCalculateForIndirectLighting) const
+FLinearColor FSpotLight::GetDirectIntensity(const FVector4f& Point, bool bCalculateForIndirectLighting) const
 {
-	FVector4 LightVector = (Point - Position).GetSafeNormal();
+	FVector4f LightVector = (Point - Position).GetSafeNormal();
 	float SpotAttenuation = FMath::Square(FMath::Clamp<float>((Dot3(LightVector, Direction) - CosOuterConeAngle) / (CosInnerConeAngle - CosOuterConeAngle),0.0f,1.0f));
 
 	if( LightFlags & GI_LIGHT_INVERSE_SQUARED )
 	{
-		FVector4 ToLight = Position - Point;
+		FVector4f ToLight = Position - Point;
 		float DistanceSqr = ToLight.SizeSquared3();
 
 		float DistanceAttenuation = 0.0f;
 		if( LightSourceLength > 0.0f )
 		{
 			// Line segment irradiance
-			FVector4 L01 = GetLightTangent() * LightSourceLength;
-			FVector4 L0 = ToLight - 0.5 * L01;
-			FVector4 L1 = ToLight + 0.5 * L01;
+			FVector4f L01 = GetLightTangent() * LightSourceLength;
+			FVector4f L0 = ToLight - 0.5 * L01;
+			FVector4f L1 = ToLight + 0.5 * L01;
 			float LengthL0 = L0.Size3();
 			float LengthL1 = L1.Size3();
 
@@ -1494,7 +1494,7 @@ FLinearColor FSpotLight::GetDirectIntensity(const FVector4& Point, bool bCalcula
 	}
 	else
 	{
-		float RadialAttenuation = FMath::Pow( FMath::Max(1.0f - ((Position - Point) / Radius).SizeSquared3(),0.0f), (FVector4::FReal)FalloffExponent );
+		float RadialAttenuation = FMath::Pow( FMath::Max(1.0f - ((Position - Point) / Radius).SizeSquared3(),0.0f), (FVector4f::FReal)FalloffExponent );
 
 		return FLight::GetDirectIntensity(Point, bCalculateForIndirectLighting) * RadialAttenuation * SpotAttenuation;
 	}
@@ -1513,14 +1513,14 @@ int32 FSpotLight::GetNumDirectPhotons(float DirectPhotonDensity) const
 }
 
 /** Generates a direction sample from the light's domain */
-void FSpotLight::SampleDirection(FLMRandomStream& RandomStream, FLightRay& SampleRay, FVector4& LightSourceNormal, FVector2D& LightSurfacePosition, float& RayPDF, FLinearColor& Power) const
+void FSpotLight::SampleDirection(FLMRandomStream& RandomStream, FLightRay& SampleRay, FVector4f& LightSourceNormal, FVector2f& LightSurfacePosition, float& RayPDF, FLinearColor& Power) const
 {
-	FVector4 XAxis(0,0,0);
-	FVector4 YAxis(0,0,0);
+	FVector4f XAxis(0,0,0);
+	FVector4f YAxis(0,0,0);
 	GenerateCoordinateSystem(Direction, XAxis, YAxis);
 
 	//@todo - the PDF should be affected by inner cone angle too
-	const FVector4 ConeSampleDirection = UniformSampleCone(RandomStream, CosOuterConeAngle, XAxis, YAxis, Direction);
+	const FVector4f ConeSampleDirection = UniformSampleCone(RandomStream, CosOuterConeAngle, XAxis, YAxis, Direction);
 
 	//@todo - take light source radius into account
 	SampleRay = FLightRay(
@@ -1545,8 +1545,8 @@ void FSpotLight::SampleDirection(
 	float& RayPDF, 
 	FLinearColor& Power) const
 {
-	FVector4 Unused;
-	FVector2D Unused2;
+	FVector4f Unused;
+	FVector2f Unused2;
 	FSpotLight::SampleDirection(RandomStream, SampleRay, Unused, Unused2, RayPDF, Power);
 }
 
@@ -1579,12 +1579,12 @@ void FRectLight::Import( FLightmassImporter& Importer )
 }
 
 
-FVector PolygonIrradiance( FVector Poly[4] )
+FVector3f PolygonIrradiance( FVector3f Poly[4] )
 {
-	FVector L0 = Poly[0].GetSafeNormal();
-	FVector L1 = Poly[1].GetSafeNormal();
-	FVector L2 = Poly[2].GetSafeNormal();
-	FVector L3 = Poly[3].GetSafeNormal();
+	FVector3f L0 = Poly[0].GetSafeNormal();
+	FVector3f L1 = Poly[1].GetSafeNormal();
+	FVector3f L2 = Poly[2].GetSafeNormal();
+	FVector3f L3 = Poly[3].GetSafeNormal();
 
 	float c01 = L0 | L1;
 	float c12 = L1 | L2;
@@ -1596,7 +1596,7 @@ FVector PolygonIrradiance( FVector Poly[4] )
 	float w23 = ( 1.5708f - 0.175f * c23 ) * FMath::InvSqrt( c23 + 1.0f );
 	float w30 = ( 1.5708f - 0.175f * c30 ) * FMath::InvSqrt( c30 + 1.0f );
 
-	FVector L;
+	FVector3f L;
 	L  = L1 ^ ( -w01 * L0 +  w12 * L2 );
 	L += L3 ^ (  w30 * L0 + -w23 * L2 );
 
@@ -1607,14 +1607,14 @@ FVector PolygonIrradiance( FVector Poly[4] )
 /**
  * Computes the intensity of the direct lighting from this light on a specific point.
  */
-FLinearColor FRectLight::GetDirectIntensity(const FVector4& Point, bool bCalculateForIndirectLighting) const
+FLinearColor FRectLight::GetDirectIntensity(const FVector4f& Point, bool bCalculateForIndirectLighting) const
 {
-	FVector ToLight = Position - Point;
+	FVector3f ToLight = Position - Point;
 
-	FVector AxisY = GetLightTangent();
-	FVector AxisZ = -Direction;
-	FVector AxisX = AxisY ^ AxisZ;
-	FVector2D Extent(
+	FVector3f AxisY = GetLightTangent();
+	FVector3f AxisZ = -Direction;
+	FVector3f AxisX = AxisY ^ AxisZ;
+	FVector2f Extent(
 		LightSourceRadius,
 		LightSourceLength
 	);
@@ -1622,7 +1622,7 @@ FLinearColor FRectLight::GetDirectIntensity(const FVector4& Point, bool bCalcula
 	if( ( ToLight | AxisZ ) < 0.0f )
 		return FLinearColor::Black;
 
-	FVector Poly[4];
+	FVector3f Poly[4];
 	Poly[0] = ToLight - AxisX * Extent.X - AxisY * Extent.Y;
 	Poly[1] = ToLight + AxisX * Extent.X - AxisY * Extent.Y;
 	Poly[2] = ToLight + AxisX * Extent.X + AxisY * Extent.Y;
@@ -1657,29 +1657,29 @@ int32 FRectLight::GetNumDirectPhotons(float DirectPhotonDensity) const
 }
 
 /** Validates a surface sample given the position that sample is affecting. */
-void FRectLight::ValidateSurfaceSample(const FVector4& Point, FLightSurfaceSample& Sample) const
+void FRectLight::ValidateSurfaceSample(const FVector4f& Point, FLightSurfaceSample& Sample) const
 {}
 
-FVector4 FRectLight::LightCenterPosition(const FVector4& ReceivingPosition, const FVector4& ReceivingNormal) const
+FVector4f FRectLight::LightCenterPosition(const FVector4f& ReceivingPosition, const FVector4f& ReceivingNormal) const
 {
-	FVector4 ToLight = Position - ReceivingPosition;
+	FVector4f ToLight = Position - ReceivingPosition;
 
-	FVector AxisY = GetLightTangent();
-	FVector AxisZ = -Direction;
-	FVector AxisX = AxisY ^ AxisZ;
-	FVector2D Extent(
+	FVector3f AxisY = GetLightTangent();
+	FVector3f AxisZ = -Direction;
+	FVector3f AxisX = AxisY ^ AxisZ;
+	FVector2f Extent(
 		LightSourceRadius,
 		LightSourceLength
 	);
 
 #if 0
-	FVector Poly[4];
+	FVector3f Poly[4];
 	Poly[0] = ToLight - AxisX * Extent.X - AxisY * Extent.Y;
 	Poly[1] = ToLight + AxisX * Extent.X - AxisY * Extent.Y;
 	Poly[2] = ToLight + AxisX * Extent.X + AxisY * Extent.Y;
 	Poly[3] = ToLight - AxisX * Extent.X + AxisY * Extent.Y;
 
-	FVector AvgDirection = PolygonIrradiance( Poly ).GetSafeNormal();
+	FVector3f AvgDirection = PolygonIrradiance( Poly ).GetSafeNormal();
 
 	// Intersect ray with plane
 	float Distance = Dot3( AxisZ, ToLight ) / Dot3( AxisZ, AvgDirection );
@@ -1688,16 +1688,16 @@ FVector4 FRectLight::LightCenterPosition(const FVector4& ReceivingPosition, cons
 	float RectLocalX = FMath::Clamp( Dot3( AxisX, -ToLight ), -Extent.X, Extent.X );
 	float RectLocalY = FMath::Clamp( Dot3( AxisY, -ToLight ), -Extent.Y, Extent.Y );
 
-	FVector4 ClosestPoint = ToLight;
+	FVector4f ClosestPoint = ToLight;
 	ClosestPoint += AxisX * RectLocalX;
 	ClosestPoint += AxisY * RectLocalY;
 	//return ClosestPoint + ReceivingPosition;
 
-	FVector4 OppositePoint = 2.0f * ToLight - ClosestPoint;
+	FVector4f OppositePoint = 2.0f * ToLight - ClosestPoint;
 
-	FVector4 L0 = ClosestPoint.GetSafeNormal();
-	FVector4 L1 = OppositePoint.GetSafeNormal();
-	FVector4 L  = ( L0 + L1 ).GetSafeNormal();
+	FVector4f L0 = ClosestPoint.GetSafeNormal();
+	FVector4f L1 = OppositePoint.GetSafeNormal();
+	FVector4f L  = ( L0 + L1 ).GetSafeNormal();
 
 	// Intersect ray with plane
 	float Distance = Dot3( AxisZ, ToLight ) / Dot3( AxisZ, L );
@@ -1705,9 +1705,9 @@ FVector4 FRectLight::LightCenterPosition(const FVector4& ReceivingPosition, cons
 #endif
 }
 
-bool FRectLight::AffectsBounds(const FBoxSphereBounds& Bounds) const
+bool FRectLight::AffectsBounds(const FBoxSphereBounds3f& Bounds) const
 {
-	FPlane Plane( Position, Direction );
+	FPlane4f Plane( Position, Direction );
 	float DistanceToPlane = ( Bounds.Origin - Position ) | Direction;
 	if( DistanceToPlane < -Bounds.SphereRadius )
 	{
@@ -1723,25 +1723,25 @@ bool FRectLight::AffectsBounds(const FBoxSphereBounds& Bounds) const
 }
 
 /** Returns true if all parts of the light are behind the surface being tested. */
-bool FRectLight::BehindSurface(const FVector4& TrianglePoint, const FVector4& TriangleNormal) const
+bool FRectLight::BehindSurface(const FVector4f& TrianglePoint, const FVector4f& TriangleNormal) const
 {
 	return false;
 }
 
 /** Gets a single direction to use for direct lighting that is representative of the whole area light. */
-FVector4 FRectLight::GetDirectLightingDirection(const FVector4& Point, const FVector4& PointNormal) const
+FVector4f FRectLight::GetDirectLightingDirection(const FVector4f& Point, const FVector4f& PointNormal) const
 {
-	FVector4 ToLight = Position - Point;
+	FVector4f ToLight = Position - Point;
 
-	FVector AxisY = GetLightTangent();
-	FVector AxisZ = -Direction;
-	FVector AxisX = AxisY ^ AxisZ;
-	FVector2D Extent(
+	FVector3f AxisY = GetLightTangent();
+	FVector3f AxisZ = -Direction;
+	FVector3f AxisX = AxisY ^ AxisZ;
+	FVector2f Extent(
 		LightSourceRadius,
 		LightSourceLength
 	);
 
-	FVector Poly[4];
+	FVector3f Poly[4];
 	Poly[0] = ToLight - AxisX * Extent.X - AxisY * Extent.Y;
 	Poly[1] = ToLight + AxisX * Extent.X - AxisY * Extent.Y;
 	Poly[2] = ToLight + AxisX * Extent.X + AxisY * Extent.Y;
@@ -1753,12 +1753,12 @@ FVector4 FRectLight::GetDirectLightingDirection(const FVector4& Point, const FVe
 /** Generates a sample on the light's surface. */
 void FRectLight::SampleLightSurface(FLMRandomStream& RandomStream, FLightSurfaceSample& Sample) const
 {
-	Sample.DiskPosition = FVector2D(0, 0);
+	Sample.DiskPosition = FVector2f(0, 0);
 
-	FVector AxisY = GetLightTangent();
-	FVector AxisZ = -Direction;
-	FVector AxisX = AxisY ^ AxisZ;
-	FVector2D Extent(
+	FVector3f AxisY = GetLightTangent();
+	FVector3f AxisZ = -Direction;
+	FVector3f AxisX = AxisY ^ AxisZ;
+	FVector2f Extent(
 		LightSourceRadius,
 		LightSourceLength
 	);
@@ -1773,12 +1773,12 @@ void FRectLight::SampleLightSurface(FLMRandomStream& RandomStream, FLightSurface
 }
 
 /** Generates a direction sample from the light's domain */
-void FRectLight::SampleDirection(FLMRandomStream& RandomStream, FLightRay& SampleRay, FVector4& LightSourceNormal, FVector2D& LightSurfacePosition, float& RayPDF, FLinearColor& Power) const
+void FRectLight::SampleDirection(FLMRandomStream& RandomStream, FLightRay& SampleRay, FVector4f& LightSourceNormal, FVector2f& LightSurfacePosition, float& RayPDF, FLinearColor& Power) const
 {
-	FVector AxisY = GetLightTangent();
-	FVector AxisZ = -Direction;
-	FVector AxisX = AxisY ^ AxisZ;
-	FVector2D Extent(
+	FVector3f AxisY = GetLightTangent();
+	FVector3f AxisZ = -Direction;
+	FVector3f AxisX = AxisY ^ AxisZ;
+	FVector2f Extent(
 		LightSourceRadius,
 		LightSourceLength
 	);
@@ -1786,9 +1786,9 @@ void FRectLight::SampleDirection(FLMRandomStream& RandomStream, FLightRay& Sampl
 	float UnitX = RandomStream.GetFraction() * 2.0f - 1.0f;
 	float UnitY = RandomStream.GetFraction() * 2.0f - 1.0f;
 
-	FVector4 SamplePosition = Position + AxisX * Extent.X * UnitX + AxisY * Extent.Y * UnitY;
+	FVector4f SamplePosition = Position + AxisX * Extent.X * UnitX + AxisY * Extent.Y * UnitY;
 
-	FVector4 SampleDirection = GetCosineHemisphereVector( RandomStream );
+	FVector4f SampleDirection = GetCosineHemisphereVector( RandomStream );
 	RayPDF = SampleDirection.Z / PI;
 
 	SampleDirection =
@@ -1816,8 +1816,8 @@ void FRectLight::SampleDirection(
 	float& RayPDF, 
 	FLinearColor& Power) const
 {
-	FVector4 Unused;
-	FVector2D Unused2;
+	FVector4f Unused;
+	FVector2f Unused2;
 	FRectLight::SampleDirection(RandomStream, SampleRay, Unused, Unused2, RayPDF, Power);
 }
 
@@ -1959,7 +1959,7 @@ void FSkyLight::ComputePrefilteredVariance()
 	}
 }
 
-FLinearColor FSkyLight::SampleRadianceCubemap(float Mip, int32 CubeFaceIndex, FVector2D FaceUV) const
+FLinearColor FSkyLight::SampleRadianceCubemap(float Mip, int32 CubeFaceIndex, FVector2f FaceUV) const
 {
 	checkSlow(bUseFilteredCubemap);
 	FLinearColor HighMipRadiance;
@@ -1987,7 +1987,7 @@ FLinearColor FSkyLight::SampleRadianceCubemap(float Mip, int32 CubeFaceIndex, FV
 	return FMath::Lerp(LowMipRadiance, HighMipRadiance, FMath::Fractional(Mip));
 }
 
-float FSkyLight::SampleVarianceCubemap(float Mip, int32 CubeFaceIndex, FVector2D FaceUV) const
+float FSkyLight::SampleVarianceCubemap(float Mip, int32 CubeFaceIndex, FVector2f FaceUV) const
 {
 	checkSlow(bUseFilteredCubemap);
 	float HighMipVariance;
@@ -2016,9 +2016,9 @@ float FSkyLight::SampleVarianceCubemap(float Mip, int32 CubeFaceIndex, FVector2D
 	return FMath::Lerp(LowMipVariance, HighMipVariance, FMath::Fractional(Mip));
 }
 
-void GetCubeFaceAndUVFromDirection(const FVector4& IncomingDirection, int32& CubeFaceIndex, FVector2D& FaceUVs)
+void GetCubeFaceAndUVFromDirection(const FVector4f& IncomingDirection, int32& CubeFaceIndex, FVector2f& FaceUVs)
 {
-	FVector AbsIncomingDirection(FMath::Abs(IncomingDirection.X), FMath::Abs(IncomingDirection.Y), FMath::Abs(IncomingDirection.Z));
+	FVector3f AbsIncomingDirection(FMath::Abs(IncomingDirection.X), FMath::Abs(IncomingDirection.Y), FMath::Abs(IncomingDirection.Z));
 
 	int32 LargestChannelIndex = 0;
 
@@ -2036,32 +2036,32 @@ void GetCubeFaceAndUVFromDirection(const FVector4& IncomingDirection, int32& Cub
 
 	if (CubeFaceIndex == 0)
 	{
-		FaceUVs = FVector2D(-IncomingDirection.Z, -IncomingDirection.Y);
+		FaceUVs = FVector2f(-IncomingDirection.Z, -IncomingDirection.Y);
 		//CubeCoordinates = float3(1, -ScaledUVs.y, -ScaledUVs.x);
 	}
 	else if (CubeFaceIndex == 1)
 	{
-		FaceUVs = FVector2D(IncomingDirection.Z, -IncomingDirection.Y);
+		FaceUVs = FVector2f(IncomingDirection.Z, -IncomingDirection.Y);
 		//CubeCoordinates = float3(-1, -ScaledUVs.y, ScaledUVs.x);
 	}
 	else if (CubeFaceIndex == 2)
 	{
-		FaceUVs = FVector2D(IncomingDirection.X, IncomingDirection.Z);
+		FaceUVs = FVector2f(IncomingDirection.X, IncomingDirection.Z);
 		//CubeCoordinates = float3(ScaledUVs.x, 1, ScaledUVs.y);
 	}
 	else if (CubeFaceIndex == 3)
 	{
-		FaceUVs = FVector2D(IncomingDirection.X, -IncomingDirection.Z);
+		FaceUVs = FVector2f(IncomingDirection.X, -IncomingDirection.Z);
 		//CubeCoordinates = float3(ScaledUVs.x, -1, -ScaledUVs.y);
 	}
 	else if (CubeFaceIndex == 4)
 	{
-		FaceUVs = FVector2D(IncomingDirection.X, -IncomingDirection.Y);
+		FaceUVs = FVector2f(IncomingDirection.X, -IncomingDirection.Y);
 		//CubeCoordinates = float3(ScaledUVs.x, -ScaledUVs.y, 1);
 	}
 	else
 	{
-		FaceUVs = FVector2D(-IncomingDirection.X, -IncomingDirection.Y);
+		FaceUVs = FVector2f(-IncomingDirection.X, -IncomingDirection.Y);
 		//CubeCoordinates = float3(-ScaledUVs.x, -ScaledUVs.y, -1);
 	}
 
@@ -2080,7 +2080,7 @@ float FSkyLight::GetMipIndexForSolidAngle(float SolidAngle) const
 	return FMath::Clamp<float>(Mip, 0.0f, NumMips - 1);
 }
 
-FLinearColor FSkyLight::GetPathLighting(const FVector4& IncomingDirection, float PathSolidAngle, bool bCalculateForIndirectLighting) const
+FLinearColor FSkyLight::GetPathLighting(const FVector4f& IncomingDirection, float PathSolidAngle, bool bCalculateForIndirectLighting) const
 {
 	if (CubemapSize == 0)
 	{
@@ -2092,7 +2092,7 @@ FLinearColor FSkyLight::GetPathLighting(const FVector4& IncomingDirection, float
 	if (bUseFilteredCubemap)
 	{
 		int32 CubeFaceIndex;
-		FVector2D FaceUVs;
+		FVector2f FaceUVs;
 		GetCubeFaceAndUVFromDirection(IncomingDirection, CubeFaceIndex, FaceUVs);
 
 		const float MipIndex = GetMipIndexForSolidAngle(PathSolidAngle);
@@ -2101,7 +2101,7 @@ FLinearColor FSkyLight::GetPathLighting(const FVector4& IncomingDirection, float
 	}
 	else
 	{
-		FSHVector3 SH = FSHVector3::SHBasisFunction(IncomingDirection);
+		FSHVector3 SH = FSHVector3::SHBasisFunction(FVector4(IncomingDirection));
 		Lighting = Dot(IrradianceEnvironmentMap, SH);
 	}
 
@@ -2115,7 +2115,7 @@ FLinearColor FSkyLight::GetPathLighting(const FVector4& IncomingDirection, float
 	return Lighting;
 }
 
-float FSkyLight::GetPathVariance(const FVector4& IncomingDirection, float PathSolidAngle) const
+float FSkyLight::GetPathVariance(const FVector4f& IncomingDirection, float PathSolidAngle) const
 {
 	if (CubemapSize == 0 || !bUseFilteredCubemap)
 	{
@@ -2123,7 +2123,7 @@ float FSkyLight::GetPathVariance(const FVector4& IncomingDirection, float PathSo
 	}
 
 	int32 CubeFaceIndex;
-	FVector2D FaceUVs;
+	FVector2f FaceUVs;
 	GetCubeFaceAndUVFromDirection(IncomingDirection, CubeFaceIndex, FaceUVs);
 
 	const float MipIndex = GetMipIndexForSolidAngle(PathSolidAngle);
@@ -2132,9 +2132,9 @@ float FSkyLight::GetPathVariance(const FVector4& IncomingDirection, float PathSo
 
 void FMeshLightPrimitive::AddSubPrimitive(const FTexelToCorners& TexelToCorners, const FIntPoint& Coordinates, const FLinearColor& InTexelPower, float NormalOffset)
 {
-	const FVector4 FirstTriangleNormal = (TexelToCorners.Corners[0].WorldPosition - TexelToCorners.Corners[1].WorldPosition) ^ (TexelToCorners.Corners[2].WorldPosition - TexelToCorners.Corners[1].WorldPosition);
+	const FVector4f FirstTriangleNormal = (TexelToCorners.Corners[0].WorldPosition - TexelToCorners.Corners[1].WorldPosition) ^ (TexelToCorners.Corners[2].WorldPosition - TexelToCorners.Corners[1].WorldPosition);
 	const float FirstTriangleArea = .5f * FirstTriangleNormal.Size3();
-	const FVector4 SecondTriangleNormal = (TexelToCorners.Corners[2].WorldPosition - TexelToCorners.Corners[1].WorldPosition) ^ (TexelToCorners.Corners[2].WorldPosition - TexelToCorners.Corners[3].WorldPosition);
+	const FVector4f SecondTriangleNormal = (TexelToCorners.Corners[2].WorldPosition - TexelToCorners.Corners[1].WorldPosition) ^ (TexelToCorners.Corners[2].WorldPosition - TexelToCorners.Corners[3].WorldPosition);
 	const float SecondTriangleArea = .5f * SecondTriangleNormal.Size3();
 	const float SubPrimitiveSurfaceArea = FirstTriangleArea + SecondTriangleArea;
 	// Convert power per texel into power per texel surface area
@@ -2144,7 +2144,7 @@ void FMeshLightPrimitive::AddSubPrimitive(const FTexelToCorners& TexelToCorners,
 	if (NumSubPrimitives == 0)
 	{
 		SurfaceNormal = TexelToCorners.WorldTangentZ;
-		const FVector4 OffsetAmount = NormalOffset * TexelToCorners.WorldTangentZ;
+		const FVector4f OffsetAmount = NormalOffset * TexelToCorners.WorldTangentZ;
 		for (int32 CornerIndex = 0; CornerIndex < NumTexelCorners; CornerIndex++)
 		{
 			Corners[CornerIndex].WorldPosition = TexelToCorners.Corners[CornerIndex].WorldPosition + OffsetAmount;
@@ -2168,7 +2168,7 @@ void FMeshLightPrimitive::AddSubPrimitive(const FTexelToCorners& TexelToCorners,
 			FIntPoint(1, 1)
 		};
 
-		const FVector4 OffsetAmount = NormalOffset * TexelToCorners.WorldTangentZ;
+		const FVector4f OffsetAmount = NormalOffset * TexelToCorners.WorldTangentZ;
 		for (int32 CornerIndex = 0; CornerIndex < NumTexelCorners; CornerIndex++)
 		{
 			const FIntPoint& ExistingFurthestCoordinates = Corners[CornerIndex].FurthestCoordinates;
@@ -2190,14 +2190,14 @@ void FMeshLightPrimitive::AddSubPrimitive(const FTexelToCorners& TexelToCorners,
 
 void FMeshLightPrimitive::Finalize()
 {
-	SurfaceNormal = SurfaceNormal.SizeSquared3() > SMALL_NUMBER ? SurfaceNormal.GetUnsafeNormal3() : FVector4(0, 0, 1);
+	SurfaceNormal = SurfaceNormal.SizeSquared3() > SMALL_NUMBER ? SurfaceNormal.GetUnsafeNormal3() : FVector4f(0, 0, 1);
 }
 
 //----------------------------------------------------------------------------
 //	Mesh Area Light class
 //----------------------------------------------------------------------------
 
-void FMeshAreaLight::Initialize(float InIndirectPhotonEmitConeAngle, const FBoxSphereBounds& InImportanceBounds)
+void FMeshAreaLight::Initialize(float InIndirectPhotonEmitConeAngle, const FBoxSphereBounds3f& InImportanceBounds)
 {
 	CosIndirectPhotonEmitConeAngle = FMath::Cos(InIndirectPhotonEmitConeAngle);
 	ImportanceBounds = InImportanceBounds;
@@ -2228,8 +2228,8 @@ void FMeshAreaLight::SetPrimitives(
 	LevelGuid = InLevelGuid;
 	TotalSurfaceArea = 0.0f;
 	TotalPower = FLinearColor::Black;
-	Position = FVector4(0,0,0);
-	FBox Bounds(ForceInit);
+	Position = FVector4f(0,0,0);
+	FBox3f Bounds(ForceInit);
 
 	CachedPrimitiveNormals.Empty(MeshAreaLightGridSize * MeshAreaLightGridSize);
 	CachedPrimitiveNormals.AddZeroed(MeshAreaLightGridSize * MeshAreaLightGridSize);
@@ -2244,7 +2244,7 @@ void FMeshAreaLight::SetPrimitives(
 		{
 			Bounds += CurrentPrimitive.Corners[CornerIndex].WorldPosition;
 		}
-		const FVector2D SphericalCoordinates = FVector(CurrentPrimitive.SurfaceNormal).UnitCartesianToSpherical();
+		const FVector2f SphericalCoordinates = FVector3f(CurrentPrimitive.SurfaceNormal).UnitCartesianToSpherical();
 		// Determine grid cell the primitive's normal falls into based on spherical coordinates
 		const int32 CacheX = FMath::Clamp(FMath::TruncToInt(SphericalCoordinates.X / (float)PI * MeshAreaLightGridSize), 0, MeshAreaLightGridSize - 1);
 		const int32 CacheY = FMath::Clamp(FMath::TruncToInt((SphericalCoordinates.Y + (float)PI) / (2 * (float)PI) * MeshAreaLightGridSize), 0, MeshAreaLightGridSize - 1);
@@ -2255,7 +2255,7 @@ void FMeshAreaLight::SetPrimitives(
 	{
 		for (int32 ThetaStep = 0; ThetaStep < MeshAreaLightGridSize; ThetaStep++)
 		{
-			const TArray<FVector4>& CurrentCachedNormals = CachedPrimitiveNormals[PhiStep * MeshAreaLightGridSize + ThetaStep];
+			const TArray<FVector4f>& CurrentCachedNormals = CachedPrimitiveNormals[PhiStep * MeshAreaLightGridSize + ThetaStep];
 			if (CurrentCachedNormals.Num() > 0)
 			{
 				OccupiedCachedPrimitiveNormalCells.Add(FIntPoint(ThetaStep, PhiStep));
@@ -2266,7 +2266,7 @@ void FMeshAreaLight::SetPrimitives(
 	// Compute the Cumulative Distribution Function for our step function of primitive surface areas
 	CalculateStep1dCDF(PrimitivePDFs, PrimitiveCDFs, UnnormalizedIntegral);
 
-	SourceBounds = FBoxSphereBounds(Bounds);
+	SourceBounds = FBoxSphereBounds3f(Bounds);
 	Position = SourceBounds.Origin;
 	Position.W = 1.0f;
 	check(TotalSurfaceArea > 0.0f);
@@ -2287,7 +2287,7 @@ void FMeshAreaLight::SetPrimitives(
  * @param Bounds - The bounding volume to test.
  * @return True if the light affects the bounding volume
  */
-bool FMeshAreaLight::AffectsBounds(const FBoxSphereBounds& Bounds) const
+bool FMeshAreaLight::AffectsBounds(const FBoxSphereBounds3f& Bounds) const
 {
 	if((Bounds.Origin - Position).SizeSquared() > FMath::Square(InfluenceRadius + Bounds.SphereRadius + SourceBounds.SphereRadius))
 	{
@@ -2305,24 +2305,24 @@ bool FMeshAreaLight::AffectsBounds(const FBoxSphereBounds& Bounds) const
 /**
  * Computes the intensity of the direct lighting from this light on a specific point.
  */
-FLinearColor FMeshAreaLight::GetDirectIntensity(const FVector4& Point, bool bCalculateForIndirectLighting) const
+FLinearColor FMeshAreaLight::GetDirectIntensity(const FVector4f& Point, bool bCalculateForIndirectLighting) const
 {
 	FLinearColor AccumulatedPower(ForceInit);
 	float AccumulatedSurfaceArea = 0.0f;
 	for (int32 PrimitiveIndex = 0; PrimitiveIndex < Primitives.Num(); PrimitiveIndex++)
 	{
 		const FMeshLightPrimitive& CurrentPrimitive = Primitives[PrimitiveIndex];
-		FVector4 PrimitiveCenter(0,0,0);
+		FVector4f PrimitiveCenter(0,0,0);
 		for (int32 CornerIndex = 0; CornerIndex < NumTexelCorners; CornerIndex++)
 		{
 			PrimitiveCenter += CurrentPrimitive.Corners[CornerIndex].WorldPosition / 4.0f;
 		}
-		const FVector4 LightVector = (Point - PrimitiveCenter).GetSafeNormal();
+		const FVector4f LightVector = (Point - PrimitiveCenter).GetSafeNormal();
 		const float NDotL = Dot3(LightVector, CurrentPrimitive.SurfaceNormal);
 		if (NDotL >= 0)
 		{
 			// Using standard Unreal attenuation for point lights for each primitive
-			const float RadialAttenuation = FMath::Pow(FMath::Max(1.0f - ((PrimitiveCenter - Point) / InfluenceRadius).SizeSquared3(), 0.0f), (FVector4::FReal)FalloffExponent);
+			const float RadialAttenuation = FMath::Pow(FMath::Max(1.0f - ((PrimitiveCenter - Point) / InfluenceRadius).SizeSquared3(), 0.0f), (FVector4f::FReal)FalloffExponent);
 			// Weight exitant power by the distance attenuation to this primitive and the light's cosine distribution around the primitive's normal
 			//@todo - photon emitting does not take the cosine distribution into account
 			AccumulatedPower += CurrentPrimitive.Power * RadialAttenuation * NDotL;
@@ -2332,7 +2332,7 @@ FLinearColor FMeshAreaLight::GetDirectIntensity(const FVector4& Point, bool bCal
 }
 
 /** Returns an intensity scale based on the receiving point. */
-float FMeshAreaLight::CustomAttenuation(const FVector4& Point, FLMRandomStream& RandomStream, bool bMaintainEvenDensity) const
+float FMeshAreaLight::CustomAttenuation(const FVector4f& Point, FLMRandomStream& RandomStream, bool bMaintainEvenDensity) const
 {
 	const float FullProbabilityDistance = .5f * InfluenceRadius;
 	float PowerWeightedAttenuation = 0.0f;
@@ -2341,7 +2341,7 @@ float FMeshAreaLight::CustomAttenuation(const FVector4& Point, FLMRandomStream&
 	for (int32 PrimitiveIndex = 0; PrimitiveIndex < Primitives.Num(); PrimitiveIndex++)
 	{
 		const FMeshLightPrimitive& CurrentPrimitive = Primitives[PrimitiveIndex];
-		FVector4 PrimitiveCenter(0,0,0);
+		FVector4f PrimitiveCenter(0,0,0);
 		for (int32 CornerIndex = 0; CornerIndex < NumTexelCorners; CornerIndex++)
 		{
 			PrimitiveCenter += CurrentPrimitive.Corners[CornerIndex].WorldPosition / 4.0f;
@@ -2349,7 +2349,7 @@ float FMeshAreaLight::CustomAttenuation(const FVector4& Point, FLMRandomStream&
 		const float NDotL = Dot3((Point - PrimitiveCenter), CurrentPrimitive.SurfaceNormal);
 		if (NDotL >= 0)
 		{
-			const float RadialAttenuation = FMath::Pow(FMath::Max(1.0f - ((PrimitiveCenter - Point) / InfluenceRadius).SizeSquared3(), 0.0f), (FVector4::FReal)FalloffExponent);
+			const float RadialAttenuation = FMath::Pow(FMath::Max(1.0f - ((PrimitiveCenter - Point) / InfluenceRadius).SizeSquared3(), 0.0f), (FVector4f::FReal)FalloffExponent);
 			const float PowerWeight = FLinearColorUtils::LinearRGBToXYZ(CurrentPrimitive.Power).G;
 			// Weight the attenuation factors by how much power this primitive emits, and its distance attenuation
 			PowerWeightedAttenuation += PowerWeight * RadialAttenuation;
@@ -2387,7 +2387,7 @@ float FMeshAreaLight::CustomAttenuation(const FVector4& Point, FLMRandomStream&
 static const float MeshAreaLightIntensityScale = 2.5f; 
 
 /** Generates a direction sample from the light's domain */
-void FMeshAreaLight::SampleDirection(FLMRandomStream& RandomStream, FLightRay& SampleRay, FVector4& LightSourceNormal, FVector2D& LightSurfacePosition, float& RayPDF, FLinearColor& Power) const
+void FMeshAreaLight::SampleDirection(FLMRandomStream& RandomStream, FLightRay& SampleRay, FVector4f& LightSourceNormal, FVector2f& LightSurfacePosition, float& RayPDF, FLinearColor& Power) const
 {
 	FLightSurfaceSample SurfaceSample;
 	FMeshAreaLight::SampleLightSurface(RandomStream, SurfaceSample);
@@ -2395,7 +2395,7 @@ void FMeshAreaLight::SampleDirection(FLMRandomStream& RandomStream, FLightRay& S
 	const float DistanceFromCenter = (SurfaceSample.Position - Position).Size3();
 
 	// Generate a sample direction from a distribution that is uniform over all directions
-	FVector4 SampleDir;
+	FVector4f SampleDir;
 	do 
 	{
 		SampleDir = GetUnitVector(RandomStream);
@@ -2437,7 +2437,7 @@ void FMeshAreaLight::SampleDirection(FLMRandomStream& RandomStream, FLightRay& S
 	{
 		const int32 ThetaStep = OccupiedCachedPrimitiveNormalCells[OccupiedCellIndex].X;
 		const int32 PhiStep = OccupiedCachedPrimitiveNormalCells[OccupiedCellIndex].Y;
-		const TArray<FVector4>& CurrentCachedNormals = CachedPrimitiveNormals[PhiStep * MeshAreaLightGridSize + ThetaStep];
+		const TArray<FVector4f>& CurrentCachedNormals = CachedPrimitiveNormals[PhiStep * MeshAreaLightGridSize + ThetaStep];
 		if (CurrentCachedNormals.Num() > 0)
 		{
 			bool bAllCornersInSameHemisphere = true;
@@ -2449,7 +2449,7 @@ void FMeshAreaLight::SampleDirection(FLMRandomStream& RandomStream, FLightRay& S
 				const float Theta = (ThetaStep + Corners[CornerIndex].X) / (float)MeshAreaLightGridSize * (float)PI;
 				const float Phi = (PhiStep + Corners[CornerIndex].Y) / (float)MeshAreaLightGridSize * 2 * (float)PI - (float)PI;
 				// Calculate the cartesian unit direction corresponding to this corner
-				const FVector4 CurrentCornerDirection = FVector2D(Theta, Phi).SphericalToUnitCartesian();
+				const FVector4f CurrentCornerDirection = FVector2f(Theta, Phi).SphericalToUnitCartesian();
 				bAllCornersInSameHemisphere = bAllCornersInSameHemisphere && Dot3(CurrentCornerDirection, SampleDir) > 0.0f;
 				bAllCornersInOppositeHemisphere = bAllCornersInOppositeHemisphere && Dot3(CurrentCornerDirection, SampleDir) < 0.0f;
 			}
@@ -2494,10 +2494,10 @@ void FMeshAreaLight::SampleDirection(
 	const int32 RayIndex = FMath::TruncToInt(RandomStream.GetFraction() * IndirectPathRays.Num());
 	checkSlow(RayIndex >= 0 && RayIndex < IndirectPathRays.Num());
 	const FIndirectPathRay& ChosenPathRay = IndirectPathRays[RayIndex];
-	const FVector4 PathRayDirection = ChosenPathRay.UnitDirection;
+	const FVector4f PathRayDirection = ChosenPathRay.UnitDirection;
 
-	FVector4 XAxis(0,0,0);
-	FVector4 YAxis(0,0,0);
+	FVector4f XAxis(0,0,0);
+	FVector4f YAxis(0,0,0);
 	GenerateCoordinateSystem(PathRayDirection, XAxis, YAxis);
 
 	// Calculate Cos of the angle between the direction and the light source normal.
@@ -2515,7 +2515,7 @@ void FMeshAreaLight::SampleDirection(
 	const float CosEmitConeAngle = FMath::Max(CosIndirectPhotonEmitConeAngle, FMath::Min(CosDirectionNormalPlaneAngle + DELTA, 1.0f));
 
 	// Generate a sample direction within a cone about the indirect path
-	const FVector4 ConeSampleDirection = UniformSampleCone(RandomStream, CosEmitConeAngle, XAxis, YAxis, PathRayDirection);
+	const FVector4f ConeSampleDirection = UniformSampleCone(RandomStream, CosEmitConeAngle, XAxis, YAxis, PathRayDirection);
 
 	FLightSurfaceSample SurfaceSample;
 	float NormalDotSampleDirection = 0.0f;
@@ -2554,7 +2554,7 @@ void FMeshAreaLight::SampleDirection(
 }
 
 /** Validates a surface sample given the position that sample is affecting. */
-void FMeshAreaLight::ValidateSurfaceSample(const FVector4& Point, FLightSurfaceSample& Sample) const
+void FMeshAreaLight::ValidateSurfaceSample(const FVector4f& Point, FLightSurfaceSample& Sample) const
 {
 }
 
@@ -2578,26 +2578,26 @@ void FMeshAreaLight::SampleLightSurface(FLMRandomStream& RandomStream, FLightSur
 	const FMeshLightPrimitive& SelectedPrimitive = Primitives[PrimitiveIndex];
 	// Approximate the primitive as a coplanar square, and sample uniformly by area
 	const float Alpha1 = RandomStream.GetFraction();
-	const FVector4 InterpolatedPosition1 = FMath::Lerp(SelectedPrimitive.Corners[0].WorldPosition, SelectedPrimitive.Corners[1].WorldPosition, Alpha1);
-	const FVector4 InterpolatedPosition2 = FMath::Lerp(SelectedPrimitive.Corners[2].WorldPosition, SelectedPrimitive.Corners[3].WorldPosition, Alpha1);
+	const FVector4f InterpolatedPosition1 = FMath::Lerp(SelectedPrimitive.Corners[0].WorldPosition, SelectedPrimitive.Corners[1].WorldPosition, Alpha1);
+	const FVector4f InterpolatedPosition2 = FMath::Lerp(SelectedPrimitive.Corners[2].WorldPosition, SelectedPrimitive.Corners[3].WorldPosition, Alpha1);
 	const float Alpha2 = RandomStream.GetFraction();
-	const FVector4 SamplePosition = FMath::Lerp(InterpolatedPosition1, InterpolatedPosition2, Alpha2);
+	const FVector4f SamplePosition = FMath::Lerp(InterpolatedPosition1, InterpolatedPosition2, Alpha2);
 	const float SamplePDF = PrimitivePDF / SelectedPrimitive.SurfaceArea;
-	Sample = FLightSurfaceSample(SamplePosition, SelectedPrimitive.SurfaceNormal, FVector2D(0,0), SamplePDF);
+	Sample = FLightSurfaceSample(SamplePosition, SelectedPrimitive.SurfaceNormal, FVector2f(0,0), SamplePDF);
 }
 
 /** Returns true if all parts of the light are behind the surface being tested. */
-bool FMeshAreaLight::BehindSurface(const FVector4& TrianglePoint, const FVector4& TriangleNormal) const
+bool FMeshAreaLight::BehindSurface(const FVector4f& TrianglePoint, const FVector4f& TriangleNormal) const
 {
 	const float NormalDotLight = Dot3(TriangleNormal, FMeshAreaLight::GetDirectLightingDirection(TrianglePoint, TriangleNormal));
 	return NormalDotLight < 0.0f;
 }
 
 /** Gets a single direction to use for direct lighting that is representative of the whole area light. */
-FVector4 FMeshAreaLight::GetDirectLightingDirection(const FVector4& Point, const FVector4& PointNormal) const
+FVector4f FMeshAreaLight::GetDirectLightingDirection(const FVector4f& Point, const FVector4f& PointNormal) const
 {
 	// The position on a sphere approximating the area light surface that will first be visible to a triangle rotating toward the light
-	const FVector4 FirstVisibleLightPoint = Position + PointNormal * SourceBounds.SphereRadius;
+	const FVector4f FirstVisibleLightPoint = Position + PointNormal * SourceBounds.SphereRadius;
 	return FirstVisibleLightPoint - Point;
 }
 
