Fix CPU Lightmass by coverting LWC types to single precision

diff --git a/Engine/Source/Programs/UnrealLightmass/Private/ImportExport/LightmassScene.h b/Engine/Source/Programs/UnrealLightmass/Private/ImportExport/LightmassScene.h
--- a/Engine/Source/Programs/UnrealLightmass/Private/ImportExport/LightmassScene.h
+++ b/Engine/Source/Programs/UnrealLightmass/Private/ImportExport/LightmassScene.h
@@ -29,17 +29,17 @@ class FLightSurfaceSample
 {
 public:
 	/** World space position */
-	FVector4 Position;
+	FVector4f Position;
     /** Normal of the light's surface at the sample point */
-	FVector4 Normal;
+	FVector4f Normal;
 	/** Position on the disk for lights modelled by a disk */
-	FVector2D DiskPosition;
+	FVector2f DiskPosition;
 	/** The probability that a sample with this position was generated */
 	float PDF;
 
 	FLightSurfaceSample() {}
 
-	FLightSurfaceSample(const FVector4& InPosition, const FVector4& InNormal, const FVector2D& InDiskPosition, float InPDF) :
+	FLightSurfaceSample(const FVector4f& InPosition, const FVector4f& InNormal, const FVector2f& InDiskPosition, float InPDF) :
 		Position(InPosition),
 		Normal(InNormal),
 		DiskPosition(InDiskPosition),
@@ -51,13 +51,13 @@ public:
 class FIndirectPathRay
 {
 public:
-	FVector4 Start;
-	FVector4 UnitDirection;
-	FVector4 LightSourceNormal;
-	FVector2D LightSurfacePosition;
+	FVector4f Start;
+	FVector4f UnitDirection;
+	FVector4f LightSourceNormal;
+	FVector2f LightSurfacePosition;
 	float Length;
 
-	FIndirectPathRay(const FVector4& InStart, const FVector4& InUnitDirection, const FVector4& InLightSourceNormal, const FVector2D& InLightSurfacePosition, const float& InLength) :
+	FIndirectPathRay(const FVector4f& InStart, const FVector4f& InUnitDirection, const FVector4f& InLightSourceNormal, const FVector2f& InLightSurfacePosition, const float& InLength) :
 		Start(InStart),
 		UnitDirection(InUnitDirection),
 		LightSourceNormal(InLightSourceNormal),
@@ -70,14 +70,14 @@ class FIrradiancePhotonData
 {
 protected:
 	/** XYZ stores world space position, W stores 1 if the photon has contribution from direct photons, and 0 otherwise. */
-	FVector4 PositionAndDirectContribution;
+	FVector4f PositionAndDirectContribution;
 
 	/** 
 	* XYZ stores the world space normal of the receiving surface.
 	* The irradiance photon caching pass stores 1 in W if the photon is actually used. 
 	* This is overwritten with the photon's irradiance in RGBE in the irradiance photon calculating pass.
 	*/
-	FVector4 SurfaceNormalAndIrradiance;
+	FVector4f SurfaceNormalAndIrradiance;
 };
 
 //----------------------------------------------------------------------------
@@ -145,20 +145,20 @@ public:
 	 * @param Bounds - The bounding volume to test.
 	 * @return True if the light affects the bounding volume
 	 */
-	virtual bool AffectsBounds(const FBoxSphereBounds& Bounds) const;
+	virtual bool AffectsBounds(const FBoxSphereBounds3f& Bounds) const;
 
-	virtual FSphere GetBoundingSphere() const;
+	virtual FSphere3f GetBoundingSphere() const;
 
 	/**
 	 * Computes the intensity of the direct lighting from this light on a specific point.
 	 */
-	virtual FLinearColor GetDirectIntensity(const FVector4& Point, bool bCalculateForIndirectLighting) const;
+	virtual FLinearColor GetDirectIntensity(const FVector4f& Point, bool bCalculateForIndirectLighting) const;
 
 	/** Returns an intensity scale based on the receiving point. */
-	virtual float CustomAttenuation(const FVector4& Point, FLMRandomStream& RandomStream, bool bMaintainEvenDensity) const { return 1.0f; }
+	virtual float CustomAttenuation(const FVector4f& Point, FLMRandomStream& RandomStream, bool bMaintainEvenDensity) const { return 1.0f; }
 
 	/** Generates a direction sample from the light's domain */
-	virtual void SampleDirection(FLMRandomStream& RandomStream, class FLightRay& SampleRay, FVector4& LightSourceNormal, FVector2D& LightSurfacePosition, float& RayPDF, FLinearColor& Power) const = 0;
+	virtual void SampleDirection(FLMRandomStream& RandomStream, class FLightRay& SampleRay, FVector4f& LightSourceNormal, FVector2f& LightSurfacePosition, float& RayPDF, FLinearColor& Power) const = 0;
 
 	/** Gives the light an opportunity to precalculate information about the indirect path rays that will be used to generate new directions. */
 	virtual void CachePathRays(const TArray<FIndirectPathRay>& IndirectPathRays) {}
@@ -181,16 +181,16 @@ public:
 	const TArray<FLightSurfaceSample>& GetCachedSurfaceSamples(int32 BounceIndex, bool bPenumbra) const;
 
 	/** Validates a surface sample given the position that sample is affecting.  The sample is unaffected by default. */
-	virtual void ValidateSurfaceSample(const FVector4& Point, FLightSurfaceSample& Sample) const {}
+	virtual void ValidateSurfaceSample(const FVector4f& Point, FLightSurfaceSample& Sample) const {}
 
 	/** Gets a single position which represents the center of the area light source from the ReceivingPosition's point of view. */
-	virtual FVector4 LightCenterPosition(const FVector4& ReceivingPosition, const FVector4& ReceivingNormal) const { return Position; }
+	virtual FVector4f LightCenterPosition(const FVector4f& ReceivingPosition, const FVector4f& ReceivingNormal) const { return Position; }
 
 	/** Returns true if all parts of the light are behind the surface being tested. */
-	virtual bool BehindSurface(const FVector4& TrianglePoint, const FVector4& TriangleNormal) const = 0;
+	virtual bool BehindSurface(const FVector4f& TrianglePoint, const FVector4f& TriangleNormal) const = 0;
 
 	/** Gets a single direction to use for direct lighting that is representative of the whole area light. */
-	virtual FVector4 GetDirectLightingDirection(const FVector4& Point, const FVector4& PointNormal) const = 0;
+	virtual FVector4f GetDirectLightingDirection(const FVector4f& Point, const FVector4f& PointNormal) const = 0;
 
 	/**
 	 * Returns whether static lighting, aka lightmaps, is being used for primitive/ light interaction.
@@ -202,7 +202,7 @@ public:
 		return (LightFlags & GI_LIGHT_HASSTATICLIGHTING) != 0;
 	}
 
-	virtual FVector GetLightTangent() const { return Direction; }
+	virtual FVector3f GetLightTangent() const { return Direction; }
 
 protected:
 	/** Cached samples of the light's surface, indexed first by bounce number, then by whether the shadow ray is a penumbra ray, then by sample index. */
@@ -240,9 +240,9 @@ public:
 	}
 
 	void Initialize(
-		const FBoxSphereBounds& InSceneBounds, 
+		const FBoxSphereBounds3f& InSceneBounds, 
 		bool bInEmitPhotonsOutsideImportanceVolume,
-		const FBoxSphereBounds& InImportanceBounds, 
+		const FBoxSphereBounds3f& InImportanceBounds, 
 		float InIndirectDiskRadius, 
 		int32 InGridSize, 
 		float InDirectPhotonDensity,
@@ -252,7 +252,7 @@ public:
 	virtual int32 GetNumDirectPhotons(float DirectPhotonDensity) const;
 
 	/** Generates a direction sample from the light's domain */
-	virtual void SampleDirection(FLMRandomStream& RandomStream, FLightRay& SampleRay, FVector4& LightSourceNormal, FVector2D& LightSurfacePosition, float& RayPDF, FLinearColor& Power) const;
+	virtual void SampleDirection(FLMRandomStream& RandomStream, FLightRay& SampleRay, FVector4f& LightSourceNormal, FVector2f& LightSurfacePosition, float& RayPDF, FLinearColor& Power) const;
 
 	/** Gives the light an opportunity to precalculate information about the indirect path rays that will be used to generate new directions. */
 	virtual void CachePathRays(const TArray<FIndirectPathRay>& IndirectPathRays);
@@ -269,16 +269,16 @@ public:
 	virtual float Power() const;
 
 	/** Validates a surface sample given the position that sample is affecting. */
-	virtual void ValidateSurfaceSample(const FVector4& Point, FLightSurfaceSample& Sample) const;
+	virtual void ValidateSurfaceSample(const FVector4f& Point, FLightSurfaceSample& Sample) const;
 
 	/** Gets a single position which represents the center of the area light source from the ReceivingPosition's point of view. */
-	virtual FVector4 LightCenterPosition(const FVector4& ReceivingPosition, const FVector4& ReceivingNormal) const;
+	virtual FVector4f LightCenterPosition(const FVector4f& ReceivingPosition, const FVector4f& ReceivingNormal) const;
 
 	/** Returns true if all parts of the light are behind the surface being tested. */
-	virtual bool BehindSurface(const FVector4& TrianglePoint, const FVector4& TriangleNormal) const;
+	virtual bool BehindSurface(const FVector4f& TrianglePoint, const FVector4f& TriangleNormal) const;
 
 	/** Gets a single direction to use for direct lighting that is representative of the whole area light. */
-	virtual FVector4 GetDirectLightingDirection(const FVector4& Point, const FVector4& PointNormal) const;
+	virtual FVector4f GetDirectLightingDirection(const FVector4f& Point, const FVector4f& PointNormal) const;
 
 protected:
 
@@ -286,7 +286,7 @@ protected:
 	float GridExtent;
 
 	/** Center of PathRayGrid in the [-1, 1] space of the directional light's disk. */
-	FVector2D GridCenter;
+	FVector2f GridCenter;
 
 	/** Size of PathRayGrid in each dimension. */
 	int32 GridSize;
@@ -295,15 +295,15 @@ protected:
 	TArray<TArray<int32>> PathRayGrid;
 
 	/** Bounds of the scene that the directional light is affecting. */
-	FBoxSphereBounds SceneBounds;
+	FBoxSphereBounds3f SceneBounds;
 
 	bool bEmitPhotonsOutsideImportanceVolume;
 
 	/** Bounds of the importance volume in the scene.  If the radius is 0, there was no importance volume. */
-	FBoxSphereBounds ImportanceBounds;
+	FBoxSphereBounds3f ImportanceBounds;
 
 	/** Center of the importance volume in the [-1,1] space of the directional light's disk */
-	FVector2D ImportanceDiskOrigin;
+	FVector2f ImportanceDiskOrigin;
 
 	/** Radius of the importance volume in the [-1,1] space of the directional light's disk */
 	float LightSpaceImportanceDiskRadius;
@@ -315,10 +315,10 @@ protected:
 	float ImportanceBoundsSampleProbability;
 
 	/** X axis of the directional light, which is unit length and orthogonal to the direction and YAxis. */
-	FVector4 XAxis;
+	FVector4f XAxis;
 
 	/** Y axis of the directional light, which is unit length and orthogonal to the direction and XAxis. */
-	FVector4 YAxis;
+	FVector4f YAxis;
 
 	/** Generates a sample on the light's surface. */
 	virtual void SampleLightSurface(FLMRandomStream& RandomStream, FLightSurfaceSample& Sample) const;
@@ -353,23 +353,23 @@ public:
 	 * @param Bounds - The bounding volume to test.
 	 * @return True if the light affects the bounding volume
 	 */
-	virtual bool AffectsBounds(const FBoxSphereBounds& Bounds) const;
+	virtual bool AffectsBounds(const FBoxSphereBounds3f& Bounds) const;
 
-	virtual FSphere GetBoundingSphere() const
+	virtual FSphere3f GetBoundingSphere() const
 	{
-		return FSphere(Position, Radius);
+		return FSphere3f(Position, Radius);
 	}
 
 	/**
 	 * Computes the intensity of the direct lighting from this light on a specific point.
 	 */
-	virtual FLinearColor GetDirectIntensity(const FVector4& Point, bool bCalculateForIndirectLighting) const;
+	virtual FLinearColor GetDirectIntensity(const FVector4f& Point, bool bCalculateForIndirectLighting) const;
 
 	/** Returns an intensity scale based on the receiving point. */
-	virtual float CustomAttenuation(const FVector4& Point, FLMRandomStream& RandomStream, bool bMaintainEvenDensity) const;
+	virtual float CustomAttenuation(const FVector4f& Point, FLMRandomStream& RandomStream, bool bMaintainEvenDensity) const;
 
 	/** Generates a direction sample from the light's domain */
-	virtual void SampleDirection(FLMRandomStream& RandomStream, FLightRay& SampleRay, FVector4& LightSourceNormal, FVector2D& LightSurfacePosition, float& RayPDF, FLinearColor& Power) const;
+	virtual void SampleDirection(FLMRandomStream& RandomStream, FLightRay& SampleRay, FVector4f& LightSourceNormal, FVector2f& LightSurfacePosition, float& RayPDF, FLinearColor& Power) const;
 
 	/** Generates a direction sample from the light based on the given rays */
 	virtual void SampleDirection(
@@ -380,20 +380,20 @@ public:
 		FLinearColor& Power) const;
 
 	/** Validates a surface sample given the position that sample is affecting. */
-	virtual void ValidateSurfaceSample(const FVector4& Point, FLightSurfaceSample& Sample) const;
+	virtual void ValidateSurfaceSample(const FVector4f& Point, FLightSurfaceSample& Sample) const;
 
 	/** Returns the light's radiant power. */
 	virtual float Power() const;
 
-	virtual FVector4 LightCenterPosition(const FVector4& ReceivingPosition, const FVector4& ReceivingNormal) const;
+	virtual FVector4f LightCenterPosition(const FVector4f& ReceivingPosition, const FVector4f& ReceivingNormal) const;
 
 	/** Returns true if all parts of the light are behind the surface being tested. */
-	virtual bool BehindSurface(const FVector4& TrianglePoint, const FVector4& TriangleNormal) const;
+	virtual bool BehindSurface(const FVector4f& TrianglePoint, const FVector4f& TriangleNormal) const;
 
 	/** Gets a single direction to use for direct lighting that is representative of the whole area light. */
-	virtual FVector4 GetDirectLightingDirection(const FVector4& Point, const FVector4& PointNormal) const;
+	virtual FVector4f GetDirectLightingDirection(const FVector4f& Point, const FVector4f& PointNormal) const;
 
-	virtual FVector GetLightTangent() const override;
+	virtual FVector3f GetLightTangent() const override;
 
 protected:
 
@@ -430,20 +430,20 @@ public:
 	 * @param Bounds - The bounding volume to test.
 	 * @return True if the light affects the bounding volume
 	 */
-	virtual bool AffectsBounds(const FBoxSphereBounds& Bounds) const;
+	virtual bool AffectsBounds(const FBoxSphereBounds3f& Bounds) const;
 
-	virtual FSphere GetBoundingSphere() const;
+	virtual FSphere3f GetBoundingSphere() const;
 
 	/**
 	 * Computes the intensity of the direct lighting from this light on a specific point.
 	 */
-	virtual FLinearColor GetDirectIntensity(const FVector4& Point, bool bCalculateForIndirectLighting) const;
+	virtual FLinearColor GetDirectIntensity(const FVector4f& Point, bool bCalculateForIndirectLighting) const;
 
 	/** Returns the number of direct photons to gather required by this light. */
 	virtual int32 GetNumDirectPhotons(float DirectPhotonDensity) const;
 
 	/** Generates a direction sample from the light's domain */
-	virtual void SampleDirection(FLMRandomStream& RandomStream, FLightRay& SampleRay, FVector4& LightSourceNormal, FVector2D& LightSurfacePosition, float& RayPDF, FLinearColor& Power) const;
+	virtual void SampleDirection(FLMRandomStream& RandomStream, FLightRay& SampleRay, FVector4f& LightSourceNormal, FVector2f& LightSurfacePosition, float& RayPDF, FLinearColor& Power) const;
 
 	/** Generates a direction sample from the light based on the given rays */
 	virtual void SampleDirection(
@@ -468,15 +468,15 @@ class FRectLight : public FPointLight, public FRectLightData
 public:
 	virtual void			Import( class FLightmassImporter& Importer );
 
-	virtual FLinearColor	GetDirectIntensity(const FVector4& Point, bool bCalculateForIndirectLighting) const override;
+	virtual FLinearColor	GetDirectIntensity(const FVector4f& Point, bool bCalculateForIndirectLighting) const override;
 	virtual int32			GetNumDirectPhotons(float DirectPhotonDensity) const override;
-	virtual void			ValidateSurfaceSample(const FVector4& Point, FLightSurfaceSample& Sample) const override;
-	virtual FVector4		LightCenterPosition(const FVector4& ReceivingPosition, const FVector4& ReceivingNormal) const override;
-	virtual bool			AffectsBounds(const FBoxSphereBounds& Bounds) const override;
-	virtual bool			BehindSurface(const FVector4& TrianglePoint, const FVector4& TriangleNormal) const override;
-	virtual FVector4		GetDirectLightingDirection(const FVector4& Point, const FVector4& PointNormal) const override;
+	virtual void			ValidateSurfaceSample(const FVector4f& Point, FLightSurfaceSample& Sample) const override;
+	virtual FVector4f		LightCenterPosition(const FVector4f& ReceivingPosition, const FVector4f& ReceivingNormal) const override;
+	virtual bool			AffectsBounds(const FBoxSphereBounds3f& Bounds) const override;
+	virtual bool			BehindSurface(const FVector4f& TrianglePoint, const FVector4f& TriangleNormal) const override;
+	virtual FVector4f		GetDirectLightingDirection(const FVector4f& Point, const FVector4f& PointNormal) const override;
 
-	virtual void			SampleDirection(FLMRandomStream& RandomStream, FLightRay& SampleRay, FVector4& LightSourceNormal, FVector2D& LightSurfacePosition, float& RayPDF, FLinearColor& Power) const override;
+	virtual void			SampleDirection(FLMRandomStream& RandomStream, FLightRay& SampleRay, FVector4f& LightSourceNormal, FVector2f& LightSurfacePosition, float& RayPDF, FLinearColor& Power) const override;
 	virtual void			SampleDirection(const TArray<FIndirectPathRay>& IndirectPathRays, FLMRandomStream& RandomStream, FLightRay& SampleRay, float& RayPDF, FLinearColor& Power) const override;
 
 protected:
@@ -511,7 +511,7 @@ public:
 	{ checkf(0, TEXT("GetNumDirectPhotons is not supported for skylights")); return 0; }
 
 	/** Generates a direction sample from the light's domain */
-	virtual void SampleDirection(FLMRandomStream& RandomStream, class FLightRay& SampleRay, FVector4& LightSourceNormal, FVector2D& LightSurfacePosition, float& RayPDF, FLinearColor& Power) const
+	virtual void SampleDirection(FLMRandomStream& RandomStream, class FLightRay& SampleRay, FVector4f& LightSourceNormal, FVector2f& LightSurfacePosition, float& RayPDF, FLinearColor& Power) const
 	{ checkf(0, TEXT("SampleDirection is not supported for skylights")); }
 
 	/** Generates a direction sample from the light based on the given rays */
@@ -528,15 +528,15 @@ public:
 	{ checkf(0, TEXT("Power is not supported for skylights")); return 0; }
 
 	/** Returns true if all parts of the light are behind the surface being tested. */
-	virtual bool BehindSurface(const FVector4& TrianglePoint, const FVector4& TriangleNormal) const { return false; }
+	virtual bool BehindSurface(const FVector4f& TrianglePoint, const FVector4f& TriangleNormal) const { return false; }
 
 	/** Gets a single direction to use for direct lighting that is representative of the whole area light. */
-	virtual FVector4 GetDirectLightingDirection(const FVector4& Point, const FVector4& PointNormal) const 
-	{ checkf(0, TEXT("GetDirectLightingDirection is not supported for skylights")); return FVector4(); }
+	virtual FVector4f GetDirectLightingDirection(const FVector4f& Point, const FVector4f& PointNormal) const 
+	{ checkf(0, TEXT("GetDirectLightingDirection is not supported for skylights")); return FVector4f(); }
 
-	FLinearColor GetPathLighting(const FVector4& IncomingDirection, float PathSolidAngle, bool bCalculateForIndirectLighting) const;
+	FLinearColor GetPathLighting(const FVector4f& IncomingDirection, float PathSolidAngle, bool bCalculateForIndirectLighting) const;
 
-	float GetPathVariance(const FVector4& IncomingDirection, float PathSolidAngle) const;
+	float GetPathVariance(const FVector4f& IncomingDirection, float PathSolidAngle) const;
 
 protected:
 
@@ -545,8 +545,8 @@ protected:
 	{ checkf(0, TEXT("SampleLightSurface is not supported for skylights")); }
 
 	float GetMipIndexForSolidAngle(float SolidAngle) const;
-	FLinearColor SampleRadianceCubemap(float MipIndex, int32 CubeFaceIndex, FVector2D FaceUV) const;
-	float SampleVarianceCubemap(float MipIndex, int32 CubeFaceIndex, FVector2D FaceUV) const;
+	FLinearColor SampleRadianceCubemap(float MipIndex, int32 CubeFaceIndex, FVector2f FaceUV) const;
+	float SampleVarianceCubemap(float MipIndex, int32 CubeFaceIndex, FVector2f FaceUV) const;
 
 	void ComputePrefilteredVariance();
 
@@ -560,7 +560,7 @@ class FMeshLightPrimitiveCorner
 {
 public:
 	/** World space corner position, not necessarily coplanar with the other corners. */
-	FVector4 WorldPosition;
+	FVector4f WorldPosition;
 	/** Coordinate in texture space corresponding to the position this corner is storing. */
 	FIntPoint FurthestCoordinates;
 };
@@ -571,7 +571,7 @@ class FMeshLightPrimitive
 public:
 	FMeshLightPrimitiveCorner Corners[NumTexelCorners];
 	/** Average normal of the sub primitives making up this simplified primitive. */
-	FVector4 SurfaceNormal;
+	FVector4f SurfaceNormal;
 	/** Radiant flux of this primitive */
 	FLinearColor Power;
 	/** Surface area of this primitive */
@@ -603,7 +603,7 @@ public:
 		return this;
 	}
 
-	void Initialize(float InIndirectPhotonEmitConeAngle, const FBoxSphereBounds& InImportanceBounds);
+	void Initialize(float InIndirectPhotonEmitConeAngle, const FBoxSphereBounds3f& InImportanceBounds);
 
 	/** Returns the number of direct photons to gather required by this light. */
 	virtual int32 GetNumDirectPhotons(float DirectPhotonDensity) const;
@@ -633,18 +633,18 @@ public:
 	* @param Bounds - The bounding volume to test.
 	* @return True if the light affects the bounding volume
 	*/
-	virtual bool AffectsBounds(const FBoxSphereBounds& Bounds) const;
+	virtual bool AffectsBounds(const FBoxSphereBounds3f& Bounds) const;
 
 	/**
 	 * Computes the intensity of the direct lighting from this light on a specific point.
 	 */
-	virtual FLinearColor GetDirectIntensity(const FVector4& Point, bool bCalculateForIndirectLighting) const;
+	virtual FLinearColor GetDirectIntensity(const FVector4f& Point, bool bCalculateForIndirectLighting) const;
 
 	/** Returns an intensity scale based on the receiving point. */
-	virtual float CustomAttenuation(const FVector4& Point, FLMRandomStream& RandomStream, bool bMaintainEvenDensity) const;
+	virtual float CustomAttenuation(const FVector4f& Point, FLMRandomStream& RandomStream, bool bMaintainEvenDensity) const;
 
 	/** Generates a direction sample from the light's domain */
-	virtual void SampleDirection(FLMRandomStream& RandomStream, FLightRay& SampleRay, FVector4& LightSourceNormal, FVector2D& LightSurfacePosition, float& RayPDF, FLinearColor& Power) const;
+	virtual void SampleDirection(FLMRandomStream& RandomStream, FLightRay& SampleRay, FVector4f& LightSourceNormal, FVector2f& LightSurfacePosition, float& RayPDF, FLinearColor& Power) const;
 
 	/** Generates a direction sample from the light based on the given rays */
 	virtual void SampleDirection(
@@ -655,16 +655,16 @@ public:
 		FLinearColor& Power) const;
 
 	/** Validates a surface sample given the position that sample is affecting. */
-	void ValidateSurfaceSample(const FVector4& Point, FLightSurfaceSample& Sample) const;
+	void ValidateSurfaceSample(const FVector4f& Point, FLightSurfaceSample& Sample) const;
 
 	/** Returns the light's radiant power. */
 	virtual float Power() const;
 
 	/** Returns true if all parts of the light are behind the surface being tested. */
-	virtual bool BehindSurface(const FVector4& TrianglePoint, const FVector4& TriangleNormal) const;
+	virtual bool BehindSurface(const FVector4f& TrianglePoint, const FVector4f& TriangleNormal) const;
 
 	/** Gets a single direction to use for direct lighting that is representative of the whole area light. */
-	virtual FVector4 GetDirectLightingDirection(const FVector4& Point, const FVector4& PointNormal) const;
+	virtual FVector4f GetDirectLightingDirection(const FVector4f& Point, const FVector4f& PointNormal) const;
 
 protected:
 	/** Radiant flux of all the primitives */
@@ -677,7 +677,7 @@ protected:
 	float InfluenceRadius;
 
 	/** Bounds of the light's primitives */
-	FBoxSphereBounds SourceBounds;
+	FBoxSphereBounds3f SourceBounds;
 
 	/** Falloff of the attenuation function */
 	float FalloffExponent;
@@ -694,13 +694,13 @@ protected:
 	FGuid LevelGuid;
 
 	/** Grid of arrays of primitive normals, used to accelerate PDF calculation once a sample is generated. */
-	TArray<TArray<FVector4> > CachedPrimitiveNormals;
+	TArray<TArray<FVector4f> > CachedPrimitiveNormals;
 
 	/** Entries in CachedPrimitiveNormals that have one or more cached normals. */
 	TArray<FIntPoint> OccupiedCachedPrimitiveNormalCells;
 
 	/** Bounds of the importance volume in the scene.  If the radius is 0, there was no importance volume. */
-	FBoxSphereBounds ImportanceBounds;
+	FBoxSphereBounds3f ImportanceBounds;
 
 	/** Probability of selecting each primitive when sampling. */
 	TArray<float> PrimitivePDFs;
@@ -721,22 +721,22 @@ protected:
 class FPrecomputedVisibilityVolume
 {
 public:
-	FBox Bounds;
-	TArray<FPlane> Planes;
+	FBox3f Bounds;
+	TArray<FPlane4f> Planes;
 };
 
 /** Volume that overrides visibility for a set of Ids. */
 class FPrecomputedVisibilityOverrideVolume
 {
 public:
-	FBox Bounds;
+	FBox3f Bounds;
 	TArray<int32> OverrideVisibilityIds;
 	TArray<int32> OverrideInvisibilityIds;
 };
 
 struct FVolumetricLightmapDensityVolume : public FVolumetricLightmapDensityVolumeData
 {
-	TArray<FPlane> Planes;
+	TArray<FPlane4f> Planes;
 };
 
 //----------------------------------------------------------------------------
@@ -749,16 +749,16 @@ public:
 	virtual ~FScene();
 
 	virtual void			Import( class FLightmassImporter& Importer );
-	FBoxSphereBounds		GetImportanceBounds() const;
+	FBoxSphereBounds3f		GetImportanceBounds() const;
 
-	FBox ImportanceBoundingBox;
-	TArray<FBox> ImportanceVolumes;
-	TArray<FBox> CharacterIndirectDetailVolumes;
+	FBox3f ImportanceBoundingBox;
+	TArray<FBox3f> ImportanceVolumes;
+	TArray<FBox3f> CharacterIndirectDetailVolumes;
 	TArray<FVolumetricLightmapDensityVolume> VolumetricLightmapDensityVolumes;
-	TArray<FSphere> Portals;
+	TArray<FSphere3f> Portals;
 	TArray<FPrecomputedVisibilityVolume> PrecomputedVisibilityVolumes;
 	TArray<FPrecomputedVisibilityOverrideVolume> PrecomputedVisibilityOverrideVolumes;
-	TArray<FVector4> CameraTrackPositions;
+	TArray<FVector4f> CameraTrackPositions;
 
 	TArray<FDirectionalLight>	DirectionalLights;
 	TArray<FPointLight>			PointLights;
@@ -787,19 +787,19 @@ public:
 	const FLight* FindLightByGuid(const FGuid& Guid) const;
 
 	/** Returns true if the specified position is inside any of the importance volumes. */
-	bool IsPointInImportanceVolume(const FVector4& Position, float Tolerance) const;
+	bool IsPointInImportanceVolume(const FVector4f& Position, float Tolerance) const;
 
-	bool IsBoxInImportanceVolume(const FBox& QueryBox) const;
+	bool IsBoxInImportanceVolume(const FBox3f& QueryBox) const;
 
 	/** Returns true if the specified position is inside any of the visibility volumes. */
-	bool IsPointInVisibilityVolume(const FVector4& Position) const;
+	bool IsPointInVisibilityVolume(const FVector4f& Position) const;
 
-	bool DoesBoxIntersectVisibilityVolume(const FBox& TestBounds) const;
+	bool DoesBoxIntersectVisibilityVolume(const FBox3f& TestBounds) const;
 
 	/** Returns accumulated bounds from all the visibility volumes. */
-	FBox GetVisibilityVolumeBounds() const;
+	FBox3f GetVisibilityVolumeBounds() const;
 
-	bool GetVolumetricLightmapAllowedMipRange(const FVector4& Position, FIntPoint& OutRange) const;
+	bool GetVolumetricLightmapAllowedMipRange(const FVector4f& Position, FIntPoint& OutRange) const;
 
 private:
 	/** Searches through all mapping arrays for the mapping matching FindGuid. */
diff --git a/Engine/Source/Programs/UnrealLightmass/Private/ImportExport/Material.h b/Engine/Source/Programs/UnrealLightmass/Private/ImportExport/Material.h
--- a/Engine/Source/Programs/UnrealLightmass/Private/ImportExport/Material.h
+++ b/Engine/Source/Programs/UnrealLightmass/Private/ImportExport/Material.h
@@ -31,21 +31,21 @@ namespace Lightmass
 		virtual void	Import( class FLightmassImporter& Importer );
 
 		/** Sample the various properties of this material at the given UV */
-		inline void SampleEmissive(const FVector2D& UV, FLinearColor& Emissive, float& OutEmissiveBoost) const
+		inline void SampleEmissive(const FVector2f& UV, FLinearColor& Emissive, float& OutEmissiveBoost) const
 		{
 			Emissive = MaterialEmissive.Sample(UV);
 			OutEmissiveBoost = EmissiveBoost;
 		}
-		inline void SampleDiffuse(const FVector2D& UV, FLinearColor& Diffuse, float& OutDiffuseBoost) const
+		inline void SampleDiffuse(const FVector2f& UV, FLinearColor& Diffuse, float& OutDiffuseBoost) const
 		{
 			Diffuse = MaterialDiffuse.Sample(UV);
 			OutDiffuseBoost = DiffuseBoost;
 		}
-		inline FLinearColor SampleTransmission(const FVector2D& UV) const
+		inline FLinearColor SampleTransmission(const FVector2f& UV) const
 		{
 			return MaterialTransmission.Sample(UV);
 		}
-		inline void SampleNormal(const FVector2D& UV, FVector4& Normal) const
+		inline void SampleNormal(const FVector2f& UV, FVector4f& Normal) const
 		{
 			Normal = MaterialNormal.SampleNormal(UV);
 			Normal.W = 0.0f;
diff --git a/Engine/Source/Programs/UnrealLightmass/Private/Launch/UnitTest.cpp b/Engine/Source/Programs/UnrealLightmass/Private/Launch/UnitTest.cpp
--- a/Engine/Source/Programs/UnrealLightmass/Private/Launch/UnitTest.cpp
+++ b/Engine/Source/Programs/UnrealLightmass/Private/Launch/UnitTest.cpp
@@ -26,7 +26,7 @@ struct FTestOctreeSemantics
 
 	FORCEINLINE static FBoxCenterAndExtent GetBoundingBox(const float& Element)
 	{
-		return FBoxCenterAndExtent(FVector4(0), FVector4(Element));
+		return FBoxCenterAndExtent(FVector4f(0), FVector4f(Element));
 	}
 
 	FORCEINLINE static bool AreElementsEqual(const float& A,const float& B)
@@ -39,7 +39,7 @@ struct FTestOctreeSemantics
 class FTestCollisionDataProvider
 {
 	TkDOPTree<FTestCollisionDataProvider, uint16>& kDOP;
-	FVector4 Vertex;
+	FVector4f Vertex;
 public:
 
 	FTestCollisionDataProvider(TkDOPTree<FTestCollisionDataProvider, uint16>& InkDOP)
@@ -53,7 +53,7 @@ public:
 	 *
 	 * @param Index the index into the vertices array
 	 */
-	FORCEINLINE const FVector4& GetVertex(uint16 Index) const
+	FORCEINLINE const FVector4f& GetVertex(uint16 Index) const
 	{
 		return Vertex;
 	}
@@ -74,25 +74,25 @@ public:
 	/**
 	 * Returns the local to world for the component
 	 */
-	FORCEINLINE const FMatrix& GetLocalToWorld(void) const
+	FORCEINLINE const FMatrix44f& GetLocalToWorld(void) const
 	{
-		return FMatrix::Identity;
+		return FMatrix44f::Identity;
 	}
 
 	/**
 	 * Returns the world to local for the component
 	 */
-	FORCEINLINE const FMatrix GetWorldToLocal(void) const
+	FORCEINLINE const FMatrix44f GetWorldToLocal(void) const
 	{
-		return FMatrix::Identity;
+		return FMatrix44f::Identity;
 	}
 
 	/**
 	 * Returns the local to world transpose adjoint for the component
 	 */
-	FORCEINLINE FMatrix GetLocalToWorldTransposeAdjoint(void) const
+	FORCEINLINE FMatrix44f GetLocalToWorldTransposeAdjoint(void) const
 	{
-		return FMatrix::Identity;
+		return FMatrix44f::Identity;
 	}
 
 	/**
@@ -154,9 +154,9 @@ void TestLightmass()
 
 	TArray<int32> ArrayCopy = TestArray;
 	
-	FVector4 TestVectorA(1, 0, 0, 1);
-	FVector4 TestVectorB(1, 1, 1, 1);
-	FVector4 TestVector = TestVectorA + TestVectorB;
+	FVector4f TestVectorA(1, 0, 0, 1);
+	FVector4f TestVectorB(1, 1, 1, 1);
+	FVector4f TestVector = TestVectorA + TestVectorB;
 
 	FString TestString = FString::Printf(TEXT("Copy has %d, Vector is [%.2f, %.2f, %.2f, %.2f]\n"), ArrayCopy[0], TestVector.X, TestVector.Y, TestVector.Z, TestVector.W);
 
@@ -167,14 +167,14 @@ void TestLightmass()
 	struct FAlignTester
 	{
 		uint8 A;
-		FMatrix M1;
+		FMatrix44f M1;
 		uint8 B;
-		FMatrix M2;
+		FMatrix44f M2;
 		uint8 C;
-		FVector4 V;
+		FVector4f V;
 	};
 	FAlignTester AlignTest;
-	checkf(((PTRINT)(&FMatrix::Identity) & 15) == 0, TEXT("Identity matrix unaligned"));
+	checkf(((PTRINT)(&FMatrix44f::Identity) & 15) == 0, TEXT("Identity matrix unaligned"));
 	checkf(((PTRINT)(&AlignTest.M1) & 15) == 0, TEXT("First matrix unaligned"));
 	checkf(((PTRINT)(&AlignTest.M2) & 15) == 0, TEXT("Second matrix unaligned"));
 	checkf(((PTRINT)(&AlignTest.V) & 15) == 0, TEXT("Vector unaligned"));
@@ -188,7 +188,7 @@ void TestLightmass()
 
 	UE_LOG(LogLightmass, Display, TEXT("Map[Five] = %d, Map[Ten] = %d"), TestMap.FindRef(TEXT("Five")), TestMap.FindRef(FString(TEXT("Ten"))));
 
-	FMatrix TestMatrix(FVector(0, 0, 0.1f), FVector(0, 1.0f, 0), FVector(0.9f, 0, 0), FVector(0, 0, 0));
+	FMatrix44f TestMatrix(FVector3f(0, 0, 0.1f), FVector3f(0, 1.0f, 0), FVector3f(0.9f, 0, 0), FVector3f(0, 0, 0));
 
 	UE_LOG(LogLightmass, Display, TEXT("Mat=\n  [%0.2f, %0.2f, %0.2f, %0.2f]\n  [%0.2f, %0.2f, %0.2f, %0.2f]\n  [%0.2f, %0.2f, %0.2f, %0.2f]\n  [%0.2f, %0.2f, %0.2f, %0.2f]"), 
 		TestMatrix.M[0][0], TestMatrix.M[0][1], TestMatrix.M[0][2], TestMatrix.M[0][3], 
@@ -217,7 +217,7 @@ void TestLightmass()
 
 	UE_LOG(LogLightmass, Display, TEXT("sizeof FDirectionalLight = %d, FLight = %d, FDirectionalLightData = %d"), sizeof(FDirectionalLight), sizeof(FLight), sizeof(FDirectionalLightData));
 
-	TOctree<float, FTestOctreeSemantics> TestOctree(FVector4(0), 10.0f);
+	TOctree<float, FTestOctreeSemantics> TestOctree(FVector4f(0), 10.0f);
 	TestOctree.AddElement(5);
 
 
@@ -226,7 +226,7 @@ void TestLightmass()
 	FTestCollisionDataProvider TestDataProvider(TestkDOP);
 	FHitResult TestResult;
 	
-	FkDOPBuildCollisionTriangle<uint16> TestTri(0, FVector4(0,0,0,0), FVector4(1,1,1,0), FVector4(2,2,2,0), INDEX_NONE, INDEX_NONE, INDEX_NONE, false, true);
+	FkDOPBuildCollisionTriangle<uint16> TestTri(0, FVector4f(0,0,0,0), FVector4f(1,1,1,0), FVector4f(2,2,2,0), INDEX_NONE, INDEX_NONE, INDEX_NONE, false, true);
 	TArray<FkDOPBuildCollisionTriangle<uint16> > TestTriangles;
 	TestTriangles.Add(TestTri);
 
diff --git a/Engine/Source/Programs/UnrealLightmass/Private/Lighting/AdaptiveVolumetricLightmap.cpp b/Engine/Source/Programs/UnrealLightmass/Private/Lighting/AdaptiveVolumetricLightmap.cpp
--- a/Engine/Source/Programs/UnrealLightmass/Private/Lighting/AdaptiveVolumetricLightmap.cpp
+++ b/Engine/Source/Programs/UnrealLightmass/Private/Lighting/AdaptiveVolumetricLightmap.cpp
@@ -11,7 +11,7 @@ namespace Lightmass
 
 struct FTriangle 
 {
-	FVector Vertices[3];
+	FVector3f Vertices[3];
 };
 
 struct FOverlapInterval
@@ -20,14 +20,14 @@ struct FOverlapInterval
 	float Max;
 };
 
-FOverlapInterval GetInterval(const FTriangle& Triangle, const FVector& Vector) 
+FOverlapInterval GetInterval(const FTriangle& Triangle, const FVector3f& Vector) 
 {
 	FOverlapInterval Result;
-	Result.Min = Result.Max = FVector::DotProduct(Vector, Triangle.Vertices[0]);
+	Result.Min = Result.Max = FVector3f::DotProduct(Vector, Triangle.Vertices[0]);
 
 	for (int32 i = 1; i < 3; ++i) 
 	{
-		float Projection = FVector::DotProduct(Vector, Triangle.Vertices[i]);
+		float Projection = FVector3f::DotProduct(Vector, Triangle.Vertices[i]);
 		Result.Min = FMath::Min(Result.Min, Projection);
 		Result.Max = FMath::Max(Result.Max, Projection);
 	}
@@ -35,26 +35,26 @@ FOverlapInterval GetInterval(const FTriangle& Triangle, const FVector& Vector)
 	return Result;
 }
 
-FOverlapInterval GetInterval(const FBox& Box, const FVector& Vector) 
+FOverlapInterval GetInterval(const FBox3f& Box, const FVector3f& Vector) 
 {
-	FVector BoxVertices[8] = 
+	FVector3f BoxVertices[8] = 
 	{
-		FVector(Box.Min.X, Box.Max.Y, Box.Max.Z),
-		FVector(Box.Min.X, Box.Max.Y, Box.Min.Z),
-		FVector(Box.Min.X, Box.Min.Y, Box.Max.Z),
-		FVector(Box.Min.X, Box.Min.Y, Box.Min.Z),
-		FVector(Box.Max.X, Box.Max.Y, Box.Max.Z),
-		FVector(Box.Max.X, Box.Max.Y, Box.Min.Z),
-		FVector(Box.Max.X, Box.Min.Y, Box.Max.Z),
-		FVector(Box.Max.X, Box.Min.Y, Box.Min.Z)
+		FVector3f(Box.Min.X, Box.Max.Y, Box.Max.Z),
+		FVector3f(Box.Min.X, Box.Max.Y, Box.Min.Z),
+		FVector3f(Box.Min.X, Box.Min.Y, Box.Max.Z),
+		FVector3f(Box.Min.X, Box.Min.Y, Box.Min.Z),
+		FVector3f(Box.Max.X, Box.Max.Y, Box.Max.Z),
+		FVector3f(Box.Max.X, Box.Max.Y, Box.Min.Z),
+		FVector3f(Box.Max.X, Box.Min.Y, Box.Max.Z),
+		FVector3f(Box.Max.X, Box.Min.Y, Box.Min.Z)
 	};
 
 	FOverlapInterval Result;
-	Result.Min = Result.Max = FVector::DotProduct(Vector, BoxVertices[0]);
+	Result.Min = Result.Max = FVector3f::DotProduct(Vector, BoxVertices[0]);
 
 	for (int32 i = 1; i < UE_ARRAY_COUNT(BoxVertices); ++i) 
 	{
-		float Projection = FVector::DotProduct(Vector, BoxVertices[i]);
+		float Projection = FVector3f::DotProduct(Vector, BoxVertices[i]);
 		Result.Min = FMath::Min(Result.Min, Projection);
 		Result.Max = FMath::Max(Result.Max, Projection);
 	}
@@ -62,41 +62,41 @@ FOverlapInterval GetInterval(const FBox& Box, const FVector& Vector)
 	return Result;
 }
 
-bool OverlapOnAxis(const FBox& Box, const FTriangle& Triangle, const FVector& Vector) 
+bool OverlapOnAxis(const FBox3f& Box, const FTriangle& Triangle, const FVector3f& Vector) 
 {
 	FOverlapInterval A = GetInterval(Box, Vector);
 	FOverlapInterval B = GetInterval(Triangle, Vector);
 	return ((B.Min <= A.Max) && (A.Min <= B.Max));
 }
 
-bool IntersectTriangleAndAABB(const FTriangle& Triangle, const FBox& Box) 
+bool IntersectTriangleAndAABB(const FTriangle& Triangle, const FBox3f& Box) 
 {
-	FVector TriangleEdge0 = Triangle.Vertices[1] - Triangle.Vertices[0]; 
-	FVector TriangleEdge1 = Triangle.Vertices[2] - Triangle.Vertices[1]; 
-	FVector TriangleEdge2 = Triangle.Vertices[0] - Triangle.Vertices[2]; 
+	FVector3f TriangleEdge0 = Triangle.Vertices[1] - Triangle.Vertices[0]; 
+	FVector3f TriangleEdge1 = Triangle.Vertices[2] - Triangle.Vertices[1]; 
+	FVector3f TriangleEdge2 = Triangle.Vertices[0] - Triangle.Vertices[2]; 
 
-	FVector BoxNormal0(1.0f, 0.0f, 0.0f);
-	FVector BoxNormal1(0.0f, 1.0f, 0.0f);
-	FVector BoxNormal2(0.0f, 0.0f, 1.0f);
+	FVector3f BoxNormal0(1.0f, 0.0f, 0.0f);
+	FVector3f BoxNormal1(0.0f, 1.0f, 0.0f);
+	FVector3f BoxNormal2(0.0f, 0.0f, 1.0f);
 
-	FVector TestDirections[13] = 
+	FVector3f TestDirections[13] = 
 	{
 		// Separating axes from the box normals
 		BoxNormal0, 
 		BoxNormal1, 
 		BoxNormal2,
 		// One separating axis for the triangle normal
-		FVector::CrossProduct(TriangleEdge0, TriangleEdge1),
+		FVector3f::CrossProduct(TriangleEdge0, TriangleEdge1),
 		// Separating axes for the triangle edges
-		FVector::CrossProduct(BoxNormal0, TriangleEdge0),
-		FVector::CrossProduct(BoxNormal0, TriangleEdge1),
-		FVector::CrossProduct(BoxNormal0, TriangleEdge2),
-		FVector::CrossProduct(BoxNormal1, TriangleEdge0),
-		FVector::CrossProduct(BoxNormal1, TriangleEdge1),
-		FVector::CrossProduct(BoxNormal1, TriangleEdge2),
-		FVector::CrossProduct(BoxNormal2, TriangleEdge0),
-		FVector::CrossProduct(BoxNormal2, TriangleEdge1),
-		FVector::CrossProduct(BoxNormal2, TriangleEdge2)
+		FVector3f::CrossProduct(BoxNormal0, TriangleEdge0),
+		FVector3f::CrossProduct(BoxNormal0, TriangleEdge1),
+		FVector3f::CrossProduct(BoxNormal0, TriangleEdge2),
+		FVector3f::CrossProduct(BoxNormal1, TriangleEdge0),
+		FVector3f::CrossProduct(BoxNormal1, TriangleEdge1),
+		FVector3f::CrossProduct(BoxNormal1, TriangleEdge2),
+		FVector3f::CrossProduct(BoxNormal2, TriangleEdge0),
+		FVector3f::CrossProduct(BoxNormal2, TriangleEdge1),
+		FVector3f::CrossProduct(BoxNormal2, TriangleEdge2)
 	};
 
 	for (int i = 0; i < UE_ARRAY_COUNT(TestDirections); ++i) 
@@ -111,21 +111,29 @@ bool IntersectTriangleAndAABB(const FTriangle& Triangle, const FBox& Box)
 	return true;
 }
 
-bool PointUnderTriangle(FVector Point, FTriangle Triangle)
+FVector3f GetBaryCentric2D(const FVector3f& Point, const FVector3f& A, const FVector3f& B, const FVector3f& C)
 {
-	const FVector BaryCentricCoordinate = FMath::GetBaryCentric2D(Point, Triangle.Vertices[0], Triangle.Vertices[1], Triangle.Vertices[2]);
+	FVector3f::FReal a = ((B.Y-C.Y)*(Point.X-C.X) + (C.X-B.X)*(Point.Y-C.Y)) / ((B.Y-C.Y)*(A.X-C.X) + (C.X-B.X)*(A.Y-C.Y));
+	FVector3f::FReal b = ((C.Y-A.Y)*(Point.X-C.X) + (A.X-C.X)*(Point.Y-C.Y)) / ((B.Y-C.Y)*(A.X-C.X) + (C.X-B.X)*(A.Y-C.Y));
+
+	return FVector3f(a, b, 1.0f - a - b);	
+}
+
+bool PointUnderTriangle(FVector3f Point, FTriangle Triangle)
+{
+	const FVector3f BaryCentricCoordinate = GetBaryCentric2D(Point, Triangle.Vertices[0], Triangle.Vertices[1], Triangle.Vertices[2]);
 	const float PointOnTriangleZ = Triangle.Vertices[0].Z * BaryCentricCoordinate.X + Triangle.Vertices[1].Z * BaryCentricCoordinate.Y + Triangle.Vertices[2].Z * BaryCentricCoordinate.Z;
 
 	return BaryCentricCoordinate.X >= 0 && BaryCentricCoordinate.Y >= 0 && BaryCentricCoordinate.Z >= 0 && PointOnTriangleZ > Point.Z;
 }
 
-bool FStaticLightingSystem::DoesVoxelIntersectSceneGeometry(const FBox& CellBounds, FGuid& OutIntersectingLevelGuid) const
+bool FStaticLightingSystem::DoesVoxelIntersectSceneGeometry(const FBox3f& CellBounds, FGuid& OutIntersectingLevelGuid) const
 {
 	const float Child2dTriangleArea = .5f * CellBounds.GetSize().X * CellBounds.GetSize().Y / (VolumetricLightmapSettings.BrickSize * VolumetricLightmapSettings.BrickSize);
 	const float SurfaceLightmapDensityThreshold = .5f * VolumetricLightmapSettings.SurfaceLightmapMinTexelsPerVoxelAxis * VolumetricLightmapSettings.SurfaceLightmapMinTexelsPerVoxelAxis / Child2dTriangleArea;
 
-	const FBox ExpandedCellBoundsSurfaceGeometry = CellBounds.ExpandBy(CellBounds.GetSize() * VolumetricLightmapSettings.VoxelizationCellExpansionForSurfaceGeometry);
-	const FBox ExpandedCellBoundsVolumeGeometry = CellBounds.ExpandBy(CellBounds.GetSize() * VolumetricLightmapSettings.VoxelizationCellExpansionForVolumeGeometry);
+	const FBox3f ExpandedCellBoundsSurfaceGeometry = CellBounds.ExpandBy(CellBounds.GetSize() * VolumetricLightmapSettings.VoxelizationCellExpansionForSurfaceGeometry);
+	const FBox3f ExpandedCellBoundsVolumeGeometry = CellBounds.ExpandBy(CellBounds.GetSize() * VolumetricLightmapSettings.VoxelizationCellExpansionForVolumeGeometry);
 
 	if (Scene.GeneralSettings.bUseFastVoxelization)
 	{
@@ -178,7 +186,7 @@ bool FStaticLightingSystem::DoesVoxelIntersectSceneGeometry(const FBox& CellBoun
 			const FStaticLightingMapping* CurrentMapping = AllMappings[MappingIndex];
 			const FStaticLightingTextureMapping* TextureMapping = CurrentMapping->GetTextureMapping();
 			const FStaticLightingMesh* CurrentMesh = CurrentMapping->Mesh;
-			const FBox& ExpandedCellBounds = CurrentMapping->GetVolumeMapping() ? ExpandedCellBoundsVolumeGeometry : ExpandedCellBoundsSurfaceGeometry;
+			const FBox3f& ExpandedCellBounds = CurrentMapping->GetVolumeMapping() ? ExpandedCellBoundsVolumeGeometry : ExpandedCellBoundsSurfaceGeometry;
 
 			const bool bMeshBelongsToLOD0 = CurrentMesh->DoesMeshBelongToLOD0();
 
@@ -201,11 +209,11 @@ bool FStaticLightingSystem::DoesVoxelIntersectSceneGeometry(const FBox& CellBoun
 						Triangle.Vertices[1] = Vertices[1].WorldPosition;
 						Triangle.Vertices[2] = Vertices[2].WorldPosition;
 
-						FBox TriangleAABB(Triangle.Vertices, UE_ARRAY_COUNT(Triangle.Vertices));
+						FBox3f TriangleAABB(Triangle.Vertices, UE_ARRAY_COUNT(Triangle.Vertices));
 
 						if (ExpandedCellBounds.Intersect(TriangleAABB))
 						{
-							const FVector4 TriangleNormal = (Vertices[2].WorldPosition - Vertices[0].WorldPosition) ^ (Vertices[1].WorldPosition - Vertices[0].WorldPosition);
+							const FVector4f TriangleNormal = (Vertices[2].WorldPosition - Vertices[0].WorldPosition) ^ (Vertices[1].WorldPosition - Vertices[0].WorldPosition);
 							const float TriangleArea = 0.5f * TriangleNormal.Size3();
 
 							if (TriangleArea > DELTA)
@@ -213,9 +221,9 @@ bool FStaticLightingSystem::DoesVoxelIntersectSceneGeometry(const FBox& CellBoun
 								if (TextureMapping)
 								{
 									// Triangle vertices in lightmap UV space, scaled by the lightmap resolution
-									const FVector2D Vertex0 = Vertices[0].TextureCoordinates[TextureMapping->LightmapTextureCoordinateIndex] * FVector2D(TextureMapping->SizeX, TextureMapping->SizeY);
-									const FVector2D Vertex1 = Vertices[1].TextureCoordinates[TextureMapping->LightmapTextureCoordinateIndex] * FVector2D(TextureMapping->SizeX, TextureMapping->SizeY);
-									const FVector2D Vertex2 = Vertices[2].TextureCoordinates[TextureMapping->LightmapTextureCoordinateIndex] * FVector2D(TextureMapping->SizeX, TextureMapping->SizeY);
+									const FVector2f Vertex0 = Vertices[0].TextureCoordinates[TextureMapping->LightmapTextureCoordinateIndex] * FVector2f(TextureMapping->SizeX, TextureMapping->SizeY);
+									const FVector2f Vertex1 = Vertices[1].TextureCoordinates[TextureMapping->LightmapTextureCoordinateIndex] * FVector2f(TextureMapping->SizeX, TextureMapping->SizeY);
+									const FVector2f Vertex2 = Vertices[2].TextureCoordinates[TextureMapping->LightmapTextureCoordinateIndex] * FVector2f(TextureMapping->SizeX, TextureMapping->SizeY);
 
 									// Area in lightmap space, or the number of lightmap texels covered by this triangle
 									const float LightmapTriangleArea = FMath::Abs(
@@ -248,7 +256,7 @@ bool FStaticLightingSystem::DoesVoxelIntersectSceneGeometry(const FBox& CellBoun
 	}
 }
 
-bool FStaticLightingSystem::ShouldRefineVoxel(int32 TreeDepth, const FBox& CellBounds, const TArray<FVector>& VoxelTestPositions, bool bDebugThisVoxel, FGuid& OutIntersectingLevelGuid) const
+bool FStaticLightingSystem::ShouldRefineVoxel(int32 TreeDepth, const FBox3f& CellBounds, const TArray<FVector3f>& VoxelTestPositions, bool bDebugThisVoxel, FGuid& OutIntersectingLevelGuid) const
 {
 	const bool bCellInsideImportanceVolume = Scene.IsBoxInImportanceVolume(CellBounds);
 
@@ -263,7 +271,7 @@ bool FStaticLightingSystem::ShouldRefineVoxel(int32 TreeDepth, const FBox& CellB
 
 	for (int32 SampleIndex = 0; SampleIndex < VoxelTestPositions.Num(); SampleIndex++)
 	{
-		FVector SamplePosition = CellBounds.Min + VoxelTestPositions[SampleIndex] * CellBounds.GetSize();
+		FVector3f SamplePosition = CellBounds.Min + VoxelTestPositions[SampleIndex] * CellBounds.GetSize();
 		FIntPoint SampleAllowedMipRange;
 		bool bSampleInDensityVolume = Scene.GetVolumetricLightmapAllowedMipRange(SamplePosition, SampleAllowedMipRange);
 		
@@ -299,8 +307,8 @@ bool FStaticLightingSystem::ShouldRefineVoxel(int32 TreeDepth, const FBox& CellB
 	
 	if (!bVoxelIntersectsScene)
 	{
-		const FBox ExpandedCellBounds = CellBounds.ExpandBy(CellBounds.GetExtent() * VolumetricLightmapSettings.VoxelizationCellExpansionForLights);
-		FBoxSphereBounds ExpandedBoxSphereBounds(ExpandedCellBounds);
+		const FBox3f ExpandedCellBounds = CellBounds.ExpandBy(CellBounds.GetExtent() * VolumetricLightmapSettings.VoxelizationCellExpansionForLights);
+		FBoxSphereBounds3f ExpandedBoxSphereBounds(ExpandedCellBounds);
 
 		for (int32 LightIndex = 0; LightIndex < Lights.Num() && !bVoxelIntersectsScene; LightIndex++)
 		{
@@ -311,7 +319,7 @@ bool FStaticLightingSystem::ShouldRefineVoxel(int32 TreeDepth, const FBox& CellB
 				// Refine around static lights, where lighting is going to be changing rapidly
 				&& Light->AffectsBounds(ExpandedBoxSphereBounds))
 			{
-				const FSphere LightBounds = Light->GetBoundingSphere();
+				const FSphere3f LightBounds = Light->GetBoundingSphere();
 
 				// If the light is smaller than the voxel, subdivide regardless of light brightness, since we will likely undersample it
 				if (LightBounds.W < ExpandedBoxSphereBounds.SphereRadius)
@@ -322,7 +330,7 @@ bool FStaticLightingSystem::ShouldRefineVoxel(int32 TreeDepth, const FBox& CellB
 				{
 					for (int32 SampleIndex = 0; SampleIndex < VoxelTestPositions.Num(); SampleIndex++)
 					{
-						FVector SamplePosition = ExpandedCellBounds.Min + VoxelTestPositions[SampleIndex] * ExpandedCellBounds.GetSize();
+						FVector3f SamplePosition = ExpandedCellBounds.Min + VoxelTestPositions[SampleIndex] * ExpandedCellBounds.GetSize();
 						FLinearColor DirectLighting = Light->GetDirectIntensity(SamplePosition, false);
 
 						if (DirectLighting.GetLuminance() > VolumetricLightmapSettings.LightBrightnessSubdivideThreshold)
@@ -343,7 +351,7 @@ bool FStaticLightingSystem::ShouldRefineVoxel(int32 TreeDepth, const FBox& CellB
 	{
 		if (Scene.GeneralSettings.bUseFastVoxelization)
 		{
-			FBox StretchedCellBounds(CellBounds.Min, FVector(CellBounds.Max.X, CellBounds.Max.Y, LandscapeCullingVoxelizationAggregateMesh->GetBounds().Max.Z));
+			FBox3f StretchedCellBounds(CellBounds.Min, FVector3f(CellBounds.Max.X, CellBounds.Max.Y, LandscapeCullingVoxelizationAggregateMesh->GetBounds().Max.Z));
 
 			if (LandscapeCullingVoxelizationAggregateMesh->IntersectBox(StretchedCellBounds) && !LandscapeCullingVoxelizationAggregateMesh->IntersectBox(CellBounds))
 			{
@@ -352,7 +360,7 @@ bool FStaticLightingSystem::ShouldRefineVoxel(int32 TreeDepth, const FBox& CellB
 		}
 		else
 		{
-			TArray<FVector, TInlineAllocator<100>> TestPositions;
+			TArray<FVector3f, TInlineAllocator<100>> TestPositions;
 			TArray<bool, TInlineAllocator<100>> PositionUnderLandscape;
 
 			int32 TestResolution = 10;
@@ -364,7 +372,7 @@ bool FStaticLightingSystem::ShouldRefineVoxel(int32 TreeDepth, const FBox& CellB
 			{
 				for (int32 X = 0; X < TestResolution; X++)
 				{
-					const FVector TestPosition = CellBounds.Min + FVector(X / (float)TestResolution, Y / (float)TestResolution, 1.0f) * CellBounds.GetSize();
+					const FVector3f TestPosition = CellBounds.Min + FVector3f(X / (float)TestResolution, Y / (float)TestResolution, 1.0f) * CellBounds.GetSize();
 					TestPositions.Add(TestPosition);
 				}
 			}
@@ -427,7 +435,7 @@ struct FIrradianceBrickBuildData
 
 // Requires texel selecting something to get into debug mode.  
 bool bDebugVolumetricLightmapCell = false;
-FVector DebugWorldPosition(548.092041f, 36.330334f, 832.141907f);
+FVector3f DebugWorldPosition(548.092041f, 36.330334f, 832.141907f);
 
 void FStaticLightingSystem::RecursivelyBuildBrickTree(
 	int32 StartCellIndex,
@@ -435,8 +443,8 @@ void FStaticLightingSystem::RecursivelyBuildBrickTree(
 	FIntVector LocalCellCoordinate, 
 	int32 TreeDepth, 
 	bool bCoveringDebugPosition,
-	const FBox& TopLevelCellBounds, 
-	const TArray<FVector>& VoxelTestPositions,
+	const FBox3f& TopLevelCellBounds, 
+	const TArray<FVector3f>& VoxelTestPositions,
 	const FGuid& IntersectingLevelGuid,
 	TArray<FIrradianceBrickBuildData>& OutBrickBuildData)
 {
@@ -464,9 +472,9 @@ void FStaticLightingSystem::RecursivelyBuildBrickTree(
 	if (TreeDepth + 1 < VolumetricLightmapSettings.MaxRefinementLevels)
 	{
 		const int32 DetailCellsPerChildLevelBrick = DetailCellsPerCurrentLevelBrick / VolumetricLightmapSettings.BrickSize;
-		const FVector BrickNormalizedMin = (FVector)LocalCellCoordinate / (float)DetailCellsPerTopLevelBrick;
-		const FVector WorldBrickMin = TopLevelCellBounds.Min + BrickNormalizedMin * TopLevelCellBounds.GetSize();
-		const FVector WorldChildCellSize = InvBrickSize * TopLevelCellBounds.GetSize() * DetailCellsPerCurrentLevelBrick / (float)DetailCellsPerTopLevelBrick;
+		const FVector3f BrickNormalizedMin = (FVector3f)LocalCellCoordinate / (float)DetailCellsPerTopLevelBrick;
+		const FVector3f WorldBrickMin = TopLevelCellBounds.Min + BrickNormalizedMin * TopLevelCellBounds.GetSize();
+		const FVector3f WorldChildCellSize = InvBrickSize * TopLevelCellBounds.GetSize() * DetailCellsPerCurrentLevelBrick / (float)DetailCellsPerTopLevelBrick;
 
 		for (int32 Z = 0; Z < VolumetricLightmapSettings.BrickSize; Z++)
 		{
@@ -478,8 +486,8 @@ void FStaticLightingSystem::RecursivelyBuildBrickTree(
 
 					if (CellIndex >= StartCellIndex && CellIndex < StartCellIndex + NumCells)
 					{
-						const FVector ChildCellPosition = WorldBrickMin + FVector(X, Y, Z) * WorldChildCellSize;
-						const FBox CellBounds(ChildCellPosition, ChildCellPosition + WorldChildCellSize);
+						const FVector3f ChildCellPosition = WorldBrickMin + FVector3f(X, Y, Z) * WorldChildCellSize;
+						const FBox3f CellBounds(ChildCellPosition, ChildCellPosition + WorldChildCellSize);
 
 						const bool bChildCoveringDebugPosition = bDebugVolumetricLightmapCell && CellBounds.IsInside(DebugWorldPosition);
 						FGuid ChildIntersectingLevelGuid;
@@ -586,13 +594,13 @@ void FStaticLightingSystem::ProcessVolumetricLightmapBrickTask(FVolumetricLightm
 		BrickData.SHCoefficients[i].AddDefaulted(TotalBrickSize);
 	}
 
-	const FVector TopLevelBrickSize = VolumetricLightmapSettings.VolumeSize / FVector(VolumetricLightmapSettings.TopLevelGridSize);
-	const FVector TopLevelBrickMin = VolumetricLightmapSettings.VolumeMin + FVector(Task->TaskIndexVector) * TopLevelBrickSize;
+	const FVector3f TopLevelBrickSize = VolumetricLightmapSettings.VolumeSize / FVector3f(VolumetricLightmapSettings.TopLevelGridSize);
+	const FVector3f TopLevelBrickMin = VolumetricLightmapSettings.VolumeMin + FVector3f(Task->TaskIndexVector) * TopLevelBrickSize;
 
-	const FVector BrickNormalizedMin = (FVector)BuildData.LocalCellCoordinate / (float)DetailCellsPerTopLevelBrick;
-	const FVector WorldBrickMin = TopLevelBrickMin + BrickNormalizedMin * TopLevelBrickSize;
+	const FVector3f BrickNormalizedMin = (FVector3f)BuildData.LocalCellCoordinate / (float)DetailCellsPerTopLevelBrick;
+	const FVector3f WorldBrickMin = TopLevelBrickMin + BrickNormalizedMin * TopLevelBrickSize;
 	const int32 DetailCellsPerCurrentLevelBrick = 1 << ((VolumetricLightmapSettings.MaxRefinementLevels - BuildData.TreeDepth) * BrickSizeLog2);
-	const FVector WorldChildCellSize = InvBrickSize * TopLevelBrickSize * DetailCellsPerCurrentLevelBrick / (float)DetailCellsPerTopLevelBrick;
+	const FVector3f WorldChildCellSize = InvBrickSize * TopLevelBrickSize * DetailCellsPerCurrentLevelBrick / (float)DetailCellsPerTopLevelBrick;
 	const int32 NumBottomLevelBricks = DetailCellsPerCurrentLevelBrick / BrickSize;
 	const float BoundarySize = NumBottomLevelBricks * InvBrickSize;
 
@@ -608,12 +616,12 @@ void FStaticLightingSystem::ProcessVolumetricLightmapBrickTask(FVolumetricLightm
 		{
 			for (int32 X = 0; X < BrickSize; X++)
 			{
-				const FVector VoxelPosition = WorldBrickMin + FVector(X, Y, Z) * WorldChildCellSize;
+				const FVector3f VoxelPosition = WorldBrickMin + FVector3f(X, Y, Z) * WorldChildCellSize;
 				
 				// Use a radius to avoid shadowing from geometry contained in the cell
-				FVolumeLightingSample CurrentSample(FVector4(VoxelPosition, WorldChildCellSize.GetMax() / 2.0f));
+				FVolumeLightingSample CurrentSample(FVector4f(VoxelPosition, WorldChildCellSize.GetMax() / 2.0f));
 
-				const FVector IndirectionCellPosition = FVector(BrickData.IndirectionTexturePosition) + FVector(X, Y, Z) * InvBrickSize * NumBottomLevelBricks;
+				const FVector3f IndirectionCellPosition = FVector3f(BrickData.IndirectionTexturePosition) + FVector3f(X, Y, Z) * InvBrickSize * NumBottomLevelBricks;
 				bool bBorderVoxel = false;
 
 				if (IndirectionCellPosition.X < BoundarySize
@@ -716,7 +724,7 @@ bool FStaticLightingSystem::ProcessVolumetricLightmapTaskIfAvailable()
 	return bAnyTaskProcessedByThisThread;
 }
 
-void FStaticLightingSystem::GenerateVoxelTestPositions(TArray<FVector>& VoxelTestPositions) const
+void FStaticLightingSystem::GenerateVoxelTestPositions(TArray<FVector3f>& VoxelTestPositions) const
 {
 	const int32 BrickSize = VolumetricLightmapSettings.BrickSize;
 	const float InvBrickSize = 1.0f / BrickSize;
@@ -732,12 +740,12 @@ void FStaticLightingSystem::GenerateVoxelTestPositions(TArray<FVector>& VoxelTes
 		{
 			for (int32 X = 0; X < BrickSize; X++)
 			{
-				const FVector BrickMin = FVector(X, Y, Z) * InvBrickSize;
+				const FVector3f BrickMin = FVector3f(X, Y, Z) * InvBrickSize;
 				const int32 CellIndex = (Z * BrickSize + Y) * BrickSize + X;
 
 				for (int32 SampleIndex = 0; SampleIndex < NumSamplesPerCell; SampleIndex++)
 				{
-					FVector RandomOffset = FVector(RandomStream.GetFraction(), RandomStream.GetFraction(), RandomStream.GetFraction()) * InvBrickSize;
+					FVector3f RandomOffset = FVector3f(RandomStream.GetFraction(), RandomStream.GetFraction(), RandomStream.GetFraction()) * InvBrickSize;
 					VoxelTestPositions[CellIndex * NumSamplesPerCell + SampleIndex] = (BrickMin + RandomOffset);
 				}
 			}
@@ -774,12 +782,12 @@ void FStaticLightingSystem::CalculateAdaptiveVolumetricLightmap(int32 TaskIndex)
 		(TopLevelBrickIndex / VolumetricLightmapSettings.TopLevelGridSize.X) % VolumetricLightmapSettings.TopLevelGridSize.Y,
 		TopLevelBrickIndex / (VolumetricLightmapSettings.TopLevelGridSize.X * VolumetricLightmapSettings.TopLevelGridSize.Y));
 
-	const FVector TopLevelBrickSize = VolumetricLightmapSettings.VolumeSize / FVector(VolumetricLightmapSettings.TopLevelGridSize);
-	const FVector TopLevelBrickMin = VolumetricLightmapSettings.VolumeMin + FVector(TaskIndexVector) * TopLevelBrickSize;
+	const FVector3f TopLevelBrickSize = VolumetricLightmapSettings.VolumeSize / FVector3f(VolumetricLightmapSettings.TopLevelGridSize);
+	const FVector3f TopLevelBrickMin = VolumetricLightmapSettings.VolumeMin + FVector3f(TaskIndexVector) * TopLevelBrickSize;
 
 	const int32 BrickSize = VolumetricLightmapSettings.BrickSize;
 
-	const FBox TopLevelBounds(TopLevelBrickMin, TopLevelBrickMin + TopLevelBrickSize);
+	const FBox3f TopLevelBounds(TopLevelBrickMin, TopLevelBrickMin + TopLevelBrickSize);
 	const bool bCoveringDebugPosition = bDebugVolumetricLightmapCell && TopLevelBounds.IsInside(DebugWorldPosition);
 
 	const int32 NumCellsPerBrick = BrickSize * BrickSize * BrickSize;
@@ -795,7 +803,7 @@ void FStaticLightingSystem::CalculateAdaptiveVolumetricLightmap(int32 TaskIndex)
 
 	check(NumCells > 0);
 
-	TArray<FVector> VoxelTestPositions;
+	TArray<FVector3f> VoxelTestPositions;
 	GenerateVoxelTestPositions(VoxelTestPositions);
 
 	TArray<FIrradianceBrickBuildData> BrickBuildData;
@@ -949,7 +957,7 @@ void FIrradianceBrickData::SetFromVolumeLightingSample(int32 Index, const FVolum
 			LQ.B.V[CoefficientIndex] = Sample.LowQualityCoefficients[CoefficientIndex][2];
 		}
 		FSHVectorRGB3 DirectLight = LQ - HQ;
-		FVector MaxLightDir( DirectLight.GetLuminance().GetMaximumDirection() );
+		FVector3f MaxLightDir( DirectLight.GetLuminance().GetMaximumDirection() );
 
 		// Set the max direction
 		FLinearColor MaxLightDirAsColor(MaxLightDir);
diff --git a/Engine/Source/Programs/UnrealLightmass/Private/Lighting/Collision.cpp b/Engine/Source/Programs/UnrealLightmass/Private/Lighting/Collision.cpp
--- a/Engine/Source/Programs/UnrealLightmass/Private/Lighting/Collision.cpp
+++ b/Engine/Source/Programs/UnrealLightmass/Private/Lighting/Collision.cpp
@@ -16,11 +16,11 @@ FStaticLightingAggregateMesh::FStaticLightingAggregateMesh(const FScene& InScene
 {
 }
 
-FBox FStaticLightingAggregateMesh::GetBounds() const
+FBox3f FStaticLightingAggregateMesh::GetBounds() const
 {
 	// Expand the bounds slightly to avoid having to handle geometry that is exactly on the bounding box,
 	// Which happens if you create a new level in Unreal with BSP from the default builder brush.
-	return bHasShadowCastingPrimitives ? SceneBounds.ExpandBy(5.0f * Scene.SceneConstants.StaticLightingLevelScale) : FBox(FVector4(0,0,0), FVector4(0,0,0));
+	return bHasShadowCastingPrimitives ? SceneBounds.ExpandBy(5.0f * Scene.SceneConstants.StaticLightingLevelScale) : FBox3f(FVector4f(0,0,0), FVector4f(0,0,0));
 }
 
 FDefaultAggregateMesh::~FDefaultAggregateMesh()
@@ -52,7 +52,7 @@ void FDefaultAggregateMesh::AddMesh(const FStaticLightingMesh* Mesh, const FStat
 		const uint32 MeshLODIndices = Mesh->GetLODIndices();
 		const uint32 MeshHLODRange = Mesh->GetHLODRange();
 
-		const FBoxSphereBounds ImportanceBounds = Scene.GetImportanceBounds();
+		const FBoxSphereBounds3f ImportanceBounds = Scene.GetImportanceBounds();
 		for(int32 TriangleIndex = 0;TriangleIndex < Mesh->NumTriangles;TriangleIndex++)
 		{
 			// Read the triangle from the mesh.
@@ -86,7 +86,7 @@ void FDefaultAggregateMesh::AddMesh(const FStaticLightingMesh* Mesh, const FStat
 			}
 
 			// Compute the triangle's normal.
-			const FVector4 TriangleNormal = (V2.WorldPosition - V0.WorldPosition) ^ (V1.WorldPosition - V0.WorldPosition);
+			const FVector4f TriangleNormal = (V2.WorldPosition - V0.WorldPosition) ^ (V1.WorldPosition - V0.WorldPosition);
 
 			// Compute the triangle area.
 			const float TriangleArea = TriangleNormal.Size3() * 0.5f;
@@ -189,7 +189,7 @@ void FDefaultAggregateMesh::AddMeshForVoxelization(const FStaticLightingMesh* Me
 				}
 
 				// Compute the triangle's normal.
-				const FVector4 TriangleNormal = (V2.WorldPosition - V0.WorldPosition) ^ (V1.WorldPosition - V0.WorldPosition);
+				const FVector4f TriangleNormal = (V2.WorldPosition - V0.WorldPosition) ^ (V1.WorldPosition - V0.WorldPosition);
 
 				// Compute the triangle area.
 				const float TriangleArea = TriangleNormal.Size3() * 0.5f;
@@ -293,17 +293,17 @@ public:
 
 	// kDOP data provider interface.
 
-	FORCEINLINE const FVector4& GetVertex(uint32 Index) const
+	FORCEINLINE const FVector4f& GetVertex(uint32 Index) const
 	{
 		return Mesh->Vertices[Index];
 	}
 
-	FORCEINLINE FVector2D GetUV(uint32 Index) const
+	FORCEINLINE FVector2f GetUV(uint32 Index) const
 	{
 		return Mesh->UVs[Index];
 	}
 
-	FORCEINLINE FVector2D GetLightmapUV(uint32 Index) const
+	FORCEINLINE FVector2f GetLightmapUV(uint32 Index) const
 	{
 		return Mesh->LightmapUVs[Index];
 	}
@@ -323,19 +323,19 @@ public:
 		return Mesh->kDopTree;
 	}
 
-	FORCEINLINE const FMatrix& GetLocalToWorld(void) const
+	FORCEINLINE const FMatrix44f& GetLocalToWorld(void) const
 	{
-		return FMatrix::Identity;
+		return FMatrix44f::Identity;
 	}
 
-	FORCEINLINE const FMatrix& GetWorldToLocal(void) const
+	FORCEINLINE const FMatrix44f& GetWorldToLocal(void) const
 	{
-		return FMatrix::Identity;
+		return FMatrix44f::Identity;
 	}
 
-	FORCEINLINE FMatrix GetLocalToWorldTransposeAdjoint(void) const
+	FORCEINLINE FMatrix44f GetLocalToWorldTransposeAdjoint(void) const
 	{
-		return FMatrix::Identity;
+		return FMatrix44f::Identity;
 	}
 
 	FORCEINLINE float GetDeterminant(void) const
@@ -457,31 +457,31 @@ bool FDefaultAggregateMesh::IntersectLightRay(
 			IntersectionVertex.WorldPosition = ClippedLightRay.Start + ClippedLightRay.Direction * ClippedLightRay.Length * Result.Time;
 			IntersectionVertex.WorldTangentZ = kDOPCheck.LocalHitNormal;
 			const FTriangleSOAPayload& Payload = TrianglePayloads[ Result.Item ];
-			const FVector4& v1 = kDOPDataProvider.GetVertex(Payload.VertexIndex[0]);
-			const FVector4& v2 = kDOPDataProvider.GetVertex(Payload.VertexIndex[1]);
-			const FVector4& v3 = kDOPDataProvider.GetVertex(Payload.VertexIndex[2]);
-			const FVector4 LocalHitPosition = kDOPCheck.LocalStart + kDOPCheck.LocalDir * Result.Time;
-			FVector4 BaryCentricWeights;
+			const FVector4f& v1 = kDOPDataProvider.GetVertex(Payload.VertexIndex[0]);
+			const FVector4f& v2 = kDOPDataProvider.GetVertex(Payload.VertexIndex[1]);
+			const FVector4f& v3 = kDOPDataProvider.GetVertex(Payload.VertexIndex[2]);
+			const FVector4f LocalHitPosition = kDOPCheck.LocalStart + kDOPCheck.LocalDir * Result.Time;
+			FVector4f BaryCentricWeights;
 			//@todo - why is such a huge tolerance needed?  Reuse the barycentric coords calculated by the ray-triangle intersection instead of deriving them from the hit position.
 			//@todo - why does this sometimes fail if there was an intersection?
 			if (bFindClosestIntersection && GetBarycentricWeights(v1, v2, v3, LocalHitPosition, KINDA_SMALL_NUMBER * 100.0f, BaryCentricWeights))
 			{
-				const FVector2D& UV1 = kDOPDataProvider.GetUV(Payload.VertexIndex[0]);
-				const FVector2D& UV2 = kDOPDataProvider.GetUV(Payload.VertexIndex[1]);
-				const FVector2D& UV3 = kDOPDataProvider.GetUV(Payload.VertexIndex[2]);
+				const FVector2f& UV1 = kDOPDataProvider.GetUV(Payload.VertexIndex[0]);
+				const FVector2f& UV2 = kDOPDataProvider.GetUV(Payload.VertexIndex[1]);
+				const FVector2f& UV3 = kDOPDataProvider.GetUV(Payload.VertexIndex[2]);
 				// Interpolate the material texture coordinates to the intersection point
 				//@todo - only lookup and interpolate UV's if needed
 				IntersectionVertex.TextureCoordinates[0] = UV1 * BaryCentricWeights.X + UV2 * BaryCentricWeights.Y + UV3 * BaryCentricWeights.Z;
-				const FVector2D& LightmapUV1 = kDOPDataProvider.GetLightmapUV(Payload.VertexIndex[0]);
-				const FVector2D& LightmapUV2 = kDOPDataProvider.GetLightmapUV(Payload.VertexIndex[1]);
-				const FVector2D& LightmapUV3 = kDOPDataProvider.GetLightmapUV(Payload.VertexIndex[2]);
+				const FVector2f& LightmapUV1 = kDOPDataProvider.GetLightmapUV(Payload.VertexIndex[0]);
+				const FVector2f& LightmapUV2 = kDOPDataProvider.GetLightmapUV(Payload.VertexIndex[1]);
+				const FVector2f& LightmapUV3 = kDOPDataProvider.GetLightmapUV(Payload.VertexIndex[2]);
 				// Interpolate the lightmap texture coordinates to the intersection point
 				IntersectionVertex.TextureCoordinates[1] = LightmapUV1 * BaryCentricWeights.X + LightmapUV2 * BaryCentricWeights.Y + LightmapUV3 * BaryCentricWeights.Z;
 			}
 			else
 			{
-				IntersectionVertex.TextureCoordinates[0] = FVector2D(0,0);
-				IntersectionVertex.TextureCoordinates[1] = FVector2D(0,0);
+				IntersectionVertex.TextureCoordinates[0] = FVector2f(0,0);
+				IntersectionVertex.TextureCoordinates[1] = FVector2f(0,0);
 			}
 
 			ClosestIntersection = FLightRayIntersection(true, IntersectionVertex, Payload.MeshInfo->Mesh, Payload.Mapping, Payload.ElementIndex);
@@ -520,7 +520,7 @@ bool FDefaultAggregateMesh::IntersectLightRay(
 	return ClosestIntersection.bIntersects;
 }
 
-const FStaticLightingMesh* FDefaultAggregateMesh::IntersectBox(const FBox Box)  const
+const FStaticLightingMesh* FDefaultAggregateMesh::IntersectBox(const FBox3f Box)  const
 {
 	FLightRay UnusedRay;
 	FStaticLightingAggregateMeshDataProvider kDOPDataProvider(this, UnusedRay);
diff --git a/Engine/Source/Programs/UnrealLightmass/Private/Lighting/Collision.h b/Engine/Source/Programs/UnrealLightmass/Private/Lighting/Collision.h
--- a/Engine/Source/Programs/UnrealLightmass/Private/Lighting/Collision.h
+++ b/Engine/Source/Programs/UnrealLightmass/Private/Lighting/Collision.h
@@ -34,9 +34,9 @@ class FLightRay
 {
 public: 
 
-	FVector4 Start;
-	FVector4 End;
-	FVector4 Direction;
+	FVector4f Start;
+	FVector4f End;
+	FVector4f Direction;
 	float Length;
 
 	/** The mapping that the ray originated from, used for conditional intersections. */
@@ -49,7 +49,7 @@ public:
 	FLightRay() {}
 
 	/** Initialization constructor. */
-	FLightRay(const FVector4& InStart, const FVector4& InEnd, const FStaticLightingMapping* InMapping, const FLight* InLight, uint32 InTraceFlags = LIGHTRAY_NONE)
+	FLightRay(const FVector4f& InStart, const FVector4f& InEnd, const FStaticLightingMapping* InMapping, const FLight* InLight, uint32 InTraceFlags = LIGHTRAY_NONE)
 	:	Start(InStart)
 	,	End(InEnd)
 	,	Direction(InEnd - InStart)
@@ -62,14 +62,14 @@ public:
 	}
 
 	/** Clips the light ray from the original start to an intersection point. */
-	void ClipAgainstIntersectionFromStart(const FVector4& IntersectionPoint)
+	void ClipAgainstIntersectionFromStart(const FVector4f& IntersectionPoint)
 	{
 		End = IntersectionPoint;
 		Direction = End - Start;
 	}
 
 	/** Clips the light ray from the original end to an intersection point. */
-	void ClipAgainstIntersectionFromEnd(const FVector4& IntersectionPoint)
+	void ClipAgainstIntersectionFromEnd(const FVector4f& IntersectionPoint)
 	{
 		Start = IntersectionPoint;
 		Direction = End - Start;
@@ -186,7 +186,7 @@ public:
 		check(0);
 	}
 
-	FBox GetBounds() const;
+	FBox3f GetBounds() const;
 
 	/** The total surface area of everything in the aggregate mesh */
 	float GetSurfaceArea() const { return SceneSurfaceArea; }
@@ -201,7 +201,7 @@ protected:
 	bool bHasShadowCastingPrimitives;
 
 	/** The bounding box of everything in the aggregate mesh. */
-	FBox SceneBounds;
+	FBox3f SceneBounds;
 
 	/** The total surface area of everything in the aggregate mesh */
 	float SceneSurfaceArea;
@@ -218,12 +218,12 @@ public:
 
 	virtual ~FDefaultAggregateMesh();
 
-	FORCEINLINE const FVector2D& GetUV(uint32 Index) const
+	FORCEINLINE const FVector2f& GetUV(uint32 Index) const
 	{
 		return UVs[Index];
 	}
 
-	FORCEINLINE const FVector2D& GetLightmapUV(uint32 Index) const
+	FORCEINLINE const FVector2f& GetLightmapUV(uint32 Index) const
 	{
 		return LightmapUVs[Index];
 	}
@@ -248,7 +248,7 @@ public:
 		class FCoherentRayCache& CoherentRayCache,
 		FLightRayIntersection& Intersection) const override;
 
-	const FStaticLightingMesh* IntersectBox(const FBox Box) const;
+	const FStaticLightingMesh* IntersectBox(const FBox3f Box) const;
 
 private:
 
@@ -270,13 +270,13 @@ private:
 	 * The vertices used by the kDOP. 
 	 * @todo - should all of these vertex attributes be stored in the same array? (ArrayOfStructures instead of SoA)
 	 */
-	TArray<FVector4> Vertices;
+	TArray<FVector4f> Vertices;
 
 	/** The texture coordinates used by the kDOP. */
-	TArray<FVector2D> UVs;
+	TArray<FVector2f> UVs;
 
 	/** The lightmap coordinates used by the kDOP. */
-	TArray<FVector2D> LightmapUVs;
+	TArray<FVector2f> LightmapUVs;
 
 	/** Needed to access properties modified by AddMesh. */
 	friend class FEmbreeVerifyAggregateMesh;
diff --git a/Engine/Source/Programs/UnrealLightmass/Private/Lighting/Embree.cpp b/Engine/Source/Programs/UnrealLightmass/Private/Lighting/Embree.cpp
--- a/Engine/Source/Programs/UnrealLightmass/Private/Lighting/Embree.cpp
+++ b/Engine/Source/Programs/UnrealLightmass/Private/Lighting/Embree.cpp
@@ -103,8 +103,8 @@ FEmbreeRay4::FEmbreeRay4(
 		primID[i] = -1;
 
 		ElementIndex[i] = -1,
-		TextureCoordinates[i] = FVector2D(0, 0);
-		LightmapCoordinates[i] = FVector2D(0, 0);
+		TextureCoordinates[i] = FVector2f(0, 0);
+		LightmapCoordinates[i] = FVector2f(0, 0);
 	}
 }
 
@@ -141,7 +141,7 @@ struct FEmbreeFilterProcessor
 	int32 Index1;
 	int32 Index2;
 
-	FVector2D TextureCoordinates; // Material Coordinates
+	FVector2f TextureCoordinates; // Material Coordinates
 
 	bool bCoordsDirty;
 
@@ -307,9 +307,9 @@ void FEmbreeFilterProcessor::UpdateCoordinates()
 			Mesh->GetInstanceableStaticMesh()->GetNonTransformedTriangleIndices(Ray.primID, Index0, Index1, Index2);
 		}
 
-		const FVector2D& UV1 = Geo.UVs[Index0];
-		const FVector2D& UV2 = Geo.UVs[Index1];
-		const FVector2D& UV3 = Geo.UVs[Index2];
+		const FVector2f& UV1 = Geo.UVs[Index0];
+		const FVector2f& UV2 = Geo.UVs[Index1];
+		const FVector2f& UV3 = Geo.UVs[Index2];
 		TextureCoordinates = UV1 * s + UV2 * Ray.u + UV3 * Ray.v;
 
 		bCoordsDirty = false;
@@ -331,9 +331,9 @@ void FEmbreeFilterProcessor::UpdateRay()
 		Ray.TextureCoordinates = TextureCoordinates;
 
 		// LightmapCoordinates
-		const FVector2D& LightmapUV1 = Geo.LightmapUVs[Index0];
-		const FVector2D& LightmapUV2 = Geo.LightmapUVs[Index1];
-		const FVector2D& LightmapUV3 = Geo.LightmapUVs[Index2];
+		const FVector2f& LightmapUV1 = Geo.LightmapUVs[Index0];
+		const FVector2f& LightmapUV2 = Geo.LightmapUVs[Index1];
+		const FVector2f& LightmapUV3 = Geo.LightmapUVs[Index2];
 		Ray.LightmapCoordinates = LightmapUV1 * s + LightmapUV2 * Ray.u + LightmapUV3 * Ray.v;
 	}
 	else if (bCoordsDirty)
@@ -437,7 +437,7 @@ void EmbreeFilterFunc4(const void* valid, void* UserPtr, RTCRay4& InRay)
 FEmbreeGeometry::FEmbreeGeometry(
 	RTCDevice EmbreeDevice, 
 	RTCScene EmbreeScene, 
-	const FBoxSphereBounds& ImportanceBounds,
+	const FBoxSphereBounds3f& ImportanceBounds,
 	const FStaticLightingMesh* InMesh,
 	const FStaticLightingMapping* InMapping,
 	bool bUseForInstancing
@@ -460,7 +460,7 @@ FEmbreeGeometry::FEmbreeGeometry(
 	UVs.AddZeroed(Mesh->NumVertices);
 	LightmapUVs.AddZeroed(Mesh->NumVertices);
 
-	FVector4* Vertices = (FVector4*) rtcMapBuffer(EmbreeScene, GeomID, RTC_VERTEX_BUFFER);
+	FVector4f* Vertices = (FVector4f*) rtcMapBuffer(EmbreeScene, GeomID, RTC_VERTEX_BUFFER);
 	int32* Indices = (int32*) rtcMapBuffer(EmbreeScene, GeomID, RTC_INDEX_BUFFER);
 
 	for (int32 TriangleIndex = 0;TriangleIndex < Mesh->NumTriangles;TriangleIndex++)
@@ -481,7 +481,7 @@ FEmbreeGeometry::FEmbreeGeometry(
 		}
 
 		// Compute the triangle's normal.
-		const FVector4 TriangleNormal = (V2.WorldPosition - V0.WorldPosition) ^ (V1.WorldPosition - V0.WorldPosition);
+		const FVector4f TriangleNormal = (V2.WorldPosition - V0.WorldPosition) ^ (V1.WorldPosition - V0.WorldPosition);
 		// Compute the triangle area.
 		const float TriangleArea = TriangleNormal.Size3() * 0.5f;
 
@@ -553,7 +553,7 @@ FEmbreeGeometry::FEmbreeGeometry(
 	check(rtcDeviceGetError(EmbreeDevice) == RTC_NO_ERROR);
 }
 
-void CalculateSurfaceArea(const FStaticLightingMesh* Mesh, const FBoxSphereBounds& ImportanceBounds, float& SurfaceArea, float& SurfaceAreaWithinImportanceVolume)
+void CalculateSurfaceArea(const FStaticLightingMesh* Mesh, const FBoxSphereBounds3f& ImportanceBounds, float& SurfaceArea, float& SurfaceAreaWithinImportanceVolume)
 {
 	SurfaceArea = 0.0f;
 	SurfaceAreaWithinImportanceVolume = 0.0f;
@@ -568,7 +568,7 @@ void CalculateSurfaceArea(const FStaticLightingMesh* Mesh, const FBoxSphereBound
 		Mesh->GetTriangle(TriangleIndex, V0, V1, V2, ElementIndex);
 
 		// Compute the triangle's normal.
-		const FVector4 TriangleNormal = (V2.WorldPosition - V0.WorldPosition) ^ (V1.WorldPosition - V0.WorldPosition);
+		const FVector4f TriangleNormal = (V2.WorldPosition - V0.WorldPosition) ^ (V1.WorldPosition - V0.WorldPosition);
 		// Compute the triangle area.
 		const float TriangleArea = TriangleNormal.Size3() * 0.5f;
 
@@ -807,13 +807,13 @@ bool FEmbreeAggregateMesh::IntersectLightRay(
 		if (EmbreeRay.instID == -1)
 		{
 			const FEmbreeGeometry& Geo = *(FEmbreeGeometry*)rtcGetUserData(EmbreeScene, EmbreeRay.geomID);
-			EmbreeVertex.WorldTangentZ = FVector(EmbreeRay.Ng[0], EmbreeRay.Ng[1], EmbreeRay.Ng[2]).GetSafeNormal();
+			EmbreeVertex.WorldTangentZ = FVector3f(EmbreeRay.Ng[0], EmbreeRay.Ng[1], EmbreeRay.Ng[2]).GetSafeNormal();
 			ClosestIntersection = FLightRayIntersection(true, EmbreeVertex, Geo.Mesh, Geo.Mapping, EmbreeRay.ElementIndex);
 		}
 		else
 		{
 			FStaticLightingMapping* Mapping = (FStaticLightingMapping*)rtcGetUserData(EmbreeScene, EmbreeRay.instID);
-			FVector GeometryNormal(EmbreeRay.Ng[0], EmbreeRay.Ng[1], EmbreeRay.Ng[2]);
+			FVector3f GeometryNormal(EmbreeRay.Ng[0], EmbreeRay.Ng[1], EmbreeRay.Ng[2]);
 			EmbreeVertex.WorldTangentZ = Mapping->Mesh->GetInstanceableStaticMesh()->LocalToWorldInverseTranspose.TransformVector(GeometryNormal).GetSafeNormal();
 			ClosestIntersection = FLightRayIntersection(true, EmbreeVertex, Mapping->Mesh, Mapping, EmbreeRay.ElementIndex);
 		}
@@ -882,7 +882,7 @@ void FEmbreeAggregateMesh::IntersectLightRays4(
 
 			FMinimalStaticLightingVertex EmbreeVertex;
 			EmbreeVertex.WorldPosition = LightRays[i].Start + LightRays[i].Direction * EmbreeRay.tfar[i];
-			EmbreeVertex.WorldTangentZ = FVector(EmbreeRay.Ngx[i], EmbreeRay.Ngy[i], EmbreeRay.Ngz[i]).GetSafeNormal();
+			EmbreeVertex.WorldTangentZ = FVector3f(EmbreeRay.Ngx[i], EmbreeRay.Ngy[i], EmbreeRay.Ngz[i]).GetSafeNormal();
 
 			EmbreeVertex.TextureCoordinates[0] = EmbreeRay.TextureCoordinates[i];
 			EmbreeVertex.TextureCoordinates[1] = EmbreeRay.LightmapCoordinates[i];
@@ -977,21 +977,21 @@ bool FEmbreeVerifyAggregateMesh::VerifyChecks(const FLightRayIntersection& Embre
 			return false;
 		}
 
-		const_cast<FVector4::FReal&>(EmbreeIntersection.IntersectionVertex.WorldPosition.W) = const_cast<FVector4::FReal&>(ClosestIntersection.IntersectionVertex.WorldPosition.W) = 1;
+		const_cast<FVector4f::FReal&>(EmbreeIntersection.IntersectionVertex.WorldPosition.W) = const_cast<FVector4f::FReal&>(ClosestIntersection.IntersectionVertex.WorldPosition.W) = 1;
 		if (!EmbreeIntersection.IntersectionVertex.WorldPosition.Equals(ClosestIntersection.IntersectionVertex.WorldPosition, .1f))
 		{
 			return false;
 		}
 
-		const_cast<FVector4::FReal&>(EmbreeIntersection.IntersectionVertex.WorldTangentZ.W) = const_cast<FVector4::FReal&>(ClosestIntersection.IntersectionVertex.WorldTangentZ.W) = 0;
+		const_cast<FVector4f::FReal&>(EmbreeIntersection.IntersectionVertex.WorldTangentZ.W) = const_cast<FVector4f::FReal&>(ClosestIntersection.IntersectionVertex.WorldTangentZ.W) = 0;
 		if (!EmbreeIntersection.IntersectionVertex.WorldTangentZ.Equals(ClosestIntersection.IntersectionVertex.WorldTangentZ, .01f))
 		{
 			return false;
 		}
 
-		FVector4 EmbreeCoord = FVector4(EmbreeIntersection.IntersectionVertex.TextureCoordinates[0].X, EmbreeIntersection.IntersectionVertex.TextureCoordinates[0].Y,
+		FVector4f EmbreeCoord = FVector4f(EmbreeIntersection.IntersectionVertex.TextureCoordinates[0].X, EmbreeIntersection.IntersectionVertex.TextureCoordinates[0].Y,
 								EmbreeIntersection.IntersectionVertex.TextureCoordinates[1].X, EmbreeIntersection.IntersectionVertex.TextureCoordinates[1].Y);
-		FVector4 ClosestCoord = FVector4(ClosestIntersection.IntersectionVertex.TextureCoordinates[0].X, ClosestIntersection.IntersectionVertex.TextureCoordinates[0].Y,
+		FVector4f ClosestCoord = FVector4f(ClosestIntersection.IntersectionVertex.TextureCoordinates[0].X, ClosestIntersection.IntersectionVertex.TextureCoordinates[0].Y,
 								ClosestIntersection.IntersectionVertex.TextureCoordinates[1].X, ClosestIntersection.IntersectionVertex.TextureCoordinates[1].Y);
 
 		if (!EmbreeCoord.Equals(ClosestCoord, .01f))
diff --git a/Engine/Source/Programs/UnrealLightmass/Private/Lighting/Embree.h b/Engine/Source/Programs/UnrealLightmass/Private/Lighting/Embree.h
--- a/Engine/Source/Programs/UnrealLightmass/Private/Lighting/Embree.h
+++ b/Engine/Source/Programs/UnrealLightmass/Private/Lighting/Embree.h
@@ -54,8 +54,8 @@ namespace Lightmass
 		// Additional Outputs.
 		// Warning: EmbreeFilterFunc must only modify these!  Nothing else is copied back to FEmbreeRay4.
 		int32 ElementIndex; // Material Index
-		FVector2D TextureCoordinates; // Material Coordinates
-		FVector2D LightmapCoordinates;
+		FVector2f TextureCoordinates; // Material Coordinates
+		FVector2f LightmapCoordinates;
 		FEmbreeTransmissionAccumulator TransmissionAcc;
 	};
 
@@ -138,8 +138,8 @@ namespace Lightmass
 		
 		// Additional Outputs.
 		int32 ElementIndex[4]; // Material Index
-		FVector2D TextureCoordinates[4]; // Material Coordinates
-		FVector2D LightmapCoordinates[4];
+		FVector2f TextureCoordinates[4]; // Material Coordinates
+		FVector2f LightmapCoordinates[4];
 		FEmbreeTransmissionAccumulator TransmissionAcc[4];
 	};
 
@@ -162,7 +162,7 @@ namespace Lightmass
 		FEmbreeGeometry(
 			RTCDevice EmbreeDevice, 
 			RTCScene EmbreeScene, 
-			const FBoxSphereBounds& ImportanceBounds,
+			const FBoxSphereBounds3f& ImportanceBounds,
 			const FStaticLightingMesh* InMesh,
 			const FStaticLightingMapping* InMapping,
 			bool bUseForInstancing
@@ -172,8 +172,8 @@ namespace Lightmass
 		const FStaticLightingMapping* Mapping;
 
 		TArray<FEmbreeTriangleDesc> TriangleDescs; // The material ID of each triangle.
-		TArray<FVector2D> UVs;
-		TArray<FVector2D> LightmapUVs;
+		TArray<FVector2f> UVs;
+		TArray<FVector2f> LightmapUVs;
 
 		uint32 GeomID; // Embree ID for this mesh.
 
