Fix CPU Lightmass by coverting LWC types to single precision

diff --git a/Engine/Source/Programs/UnrealLightmass/Private/Lighting/FinalGather.cpp b/Engine/Source/Programs/UnrealLightmass/Private/Lighting/FinalGather.cpp
--- a/Engine/Source/Programs/UnrealLightmass/Private/Lighting/FinalGather.cpp
+++ b/Engine/Source/Programs/UnrealLightmass/Private/Lighting/FinalGather.cpp
@@ -21,11 +21,11 @@ namespace Lightmass
 {
 
 void FStaticLightingSystem::GatherVolumeImportancePhotonDirections(
-	const FVector WorldPosition,
-	FVector FirstHemisphereNormal,
-	FVector SecondHemisphereNormal,
-	TArray<FVector4>& FirstHemisphereImportancePhotonDirections,
-	TArray<FVector4>& SecondHemisphereImportancePhotonDirections,
+	const FVector3f WorldPosition,
+	FVector3f FirstHemisphereNormal,
+	FVector3f SecondHemisphereNormal,
+	TArray<FVector4f>& FirstHemisphereImportancePhotonDirections,
+	TArray<FVector4f>& SecondHemisphereImportancePhotonDirections,
 	bool bDebugThisSample) const
 {
 	if (GeneralSettings.NumIndirectLightingBounces > 0 && PhotonMappingSettings.bUsePhotonMapping && PhotonMappingSettings.bUsePhotonSegmentsForVolumeLighting)
@@ -52,7 +52,7 @@ void FStaticLightingSystem::GatherVolumeImportancePhotonDirections(
 			const FPhoton& CurrentPhoton = *FoundPhotonSegments[PhotonIndex].Photon;
 			// Calculate the direction from the current position to the photon's source
 			// Using the photon's incident direction unmodified produces artifacts proportional to the distance to that photon
-			const FVector4 NewDirection = CurrentPhoton.GetPosition() + CurrentPhoton.GetIncidentDirection() * CurrentPhoton.GetDistance() - WorldPosition;
+			const FVector4f NewDirection = CurrentPhoton.GetPosition() + CurrentPhoton.GetIncidentDirection() * CurrentPhoton.GetDistance() - WorldPosition;
 			// Only use the direction if it is in the hemisphere of the normal
 			// FindNearbyPhotons only returns photons whose incident directions lie in this hemisphere, but the recalculated direction might not.
 			if (Dot3(NewDirection, FirstHemisphereNormal) > 0.0f)
@@ -70,10 +70,10 @@ void FStaticLightingSystem::GatherVolumeImportancePhotonDirections(
 
 /** Calculates incident radiance for a given world space position. */
 void FStaticLightingSystem::CalculateVolumeSampleIncidentRadiance(
-	const TArray<FVector4>& UniformHemisphereSamples,
-	const TArray<FVector2D>& UniformHemisphereSampleUniforms,
+	const TArray<FVector4f>& UniformHemisphereSamples,
+	const TArray<FVector2f>& UniformHemisphereSampleUniforms,
 	float MaxUnoccludedLength,
-	const TArray<FVector, TInlineAllocator<1>>& VertexOffsets,
+	const TArray<FVector3f, TInlineAllocator<1>>& VertexOffsets,
 	FVolumeLightingSample& LightingSample,
 	float& OutBackfacingHitsFraction,
 	float& OutMinDistanceToSurface,
@@ -84,15 +84,15 @@ void FStaticLightingSystem::CalculateVolumeSampleIncidentRadiance(
 {
 	const double StartTime = FPlatformTime::Seconds();
 
-	const FVector4 Position = LightingSample.GetPosition();
+	const FVector4f Position = LightingSample.GetPosition();
 
-	TArray<FVector4> UpperHemisphereImportancePhotonDirections;
-	TArray<FVector4> LowerHemisphereImportancePhotonDirections;
+	TArray<FVector4f> UpperHemisphereImportancePhotonDirections;
+	TArray<FVector4f> LowerHemisphereImportancePhotonDirections;
 	
 	GatherVolumeImportancePhotonDirections(
 		Position,
-		FVector(0, 0, 1),
-		FVector(0, 0, -1),
+		FVector3f(0, 0, 1),
+		FVector3f(0, 0, -1),
 		UpperHemisphereImportancePhotonDirections,
 		LowerHemisphereImportancePhotonDirections,
 		bDebugThisSample);
@@ -101,7 +101,7 @@ void FStaticLightingSystem::CalculateVolumeSampleIncidentRadiance(
 	{
 		FScopeLock DebugOutputLock(&DebugOutputSync);
 		FDebugStaticLightingVertex DebugVertex;
-		DebugVertex.VertexNormal = FVector(0, 0, 1);
+		DebugVertex.VertexNormal = FVector3f(0, 0, 1);
 		DebugVertex.VertexPosition = LightingSample.GetPosition();
 		MappingContext.DebugOutput->Vertices.Add(DebugVertex);
 	}
@@ -111,11 +111,11 @@ void FStaticLightingSystem::CalculateVolumeSampleIncidentRadiance(
 
 	FFullStaticLightingVertex RepresentativeVertex;
 	RepresentativeVertex.WorldPosition = Position;
-	RepresentativeVertex.TextureCoordinates[0] = FVector2D(0,0);
-	RepresentativeVertex.TextureCoordinates[1] = FVector2D(0,0);
+	RepresentativeVertex.TextureCoordinates[0] = FVector2f(0,0);
+	RepresentativeVertex.TextureCoordinates[1] = FVector2f(0,0);
 
 	// Construct a vertex to capture incident radiance for the positive Z hemisphere
-	RepresentativeVertex.WorldTangentZ = RepresentativeVertex.TriangleNormal = FVector4(0,0,1);
+	RepresentativeVertex.WorldTangentZ = RepresentativeVertex.TriangleNormal = FVector4f(0,0,1);
 	RepresentativeVertex.GenerateVertexTangents();
 	RepresentativeVertex.GenerateTriangleTangents();
 
@@ -158,7 +158,7 @@ void FStaticLightingSystem::CalculateVolumeSampleIncidentRadiance(
 	MappingContext.Stats.VolumetricLightmapFinalGatherTime += EndUpperFinalGatherTime - EndUpperDirectLightingTime;
 
 	// Construct a vertex to capture incident radiance for the negative Z hemisphere
-	RepresentativeVertex.WorldTangentZ = RepresentativeVertex.TriangleNormal = FVector4(0,0,-1);
+	RepresentativeVertex.WorldTangentZ = RepresentativeVertex.TriangleNormal = FVector4f(0,0,-1);
 	RepresentativeVertex.GenerateVertexTangents();
 	RepresentativeVertex.GenerateTriangleTangents();
 
@@ -226,13 +226,13 @@ void FStaticLightingSystem::CalculateVolumeSampleIncidentRadiance(
 }
 
 /** Returns environment lighting for the given direction. */
-FLinearColor FStaticLightingSystem::EvaluateEnvironmentLighting(const FVector4& IncomingDirection) const
+FLinearColor FStaticLightingSystem::EvaluateEnvironmentLighting(const FVector4f& IncomingDirection) const
 {
 	// Upper hemisphere only
 	return IncomingDirection.Z < 0 ? (MaterialSettings.EnvironmentColor / (float)PI) : FLinearColor::Black;
 }
 
-void FStaticLightingSystem::EvaluateSkyLighting(const FVector4& IncomingDirection, float PathSolidAngle, bool bShadowed, bool bForDirectLighting, FLinearColor& OutStaticLighting, FLinearColor& OutStationaryLighting) const
+void FStaticLightingSystem::EvaluateSkyLighting(const FVector4f& IncomingDirection, float PathSolidAngle, bool bShadowed, bool bForDirectLighting, FLinearColor& OutStaticLighting, FLinearColor& OutStationaryLighting) const
 {
 	for (int32 LightIndex = 0; LightIndex < SkyLights.Num(); LightIndex++)
 	{
@@ -254,7 +254,7 @@ void FStaticLightingSystem::EvaluateSkyLighting(const FVector4& IncomingDirectio
 	}
 }
 
-float FStaticLightingSystem::EvaluateSkyVariance(const FVector4& IncomingDirection, float PathSolidAngle) const
+float FStaticLightingSystem::EvaluateSkyVariance(const FVector4f& IncomingDirection, float PathSolidAngle) const
 {
 	float Variance = 0;
 
@@ -274,7 +274,7 @@ FLinearColor FStaticLightingSystem::CalculateExitantRadiance(
 	const FStaticLightingMesh* HitMesh,
 	const FMinimalStaticLightingVertex& Vertex,
 	int32 ElementIndex,
-	const FVector4& OutgoingDirection,
+	const FVector4f& OutgoingDirection,
 	int32 BounceNumber,
 	EHemisphereGatherClassification GatherClassification,
 	FStaticLightingMappingContext& MappingContext,
@@ -318,8 +318,8 @@ void FStaticLightingSystem::IntersectLightRays(
 	const FFullStaticLightingVertex& Vertex,
 	float SampleRadius,
 	int32 NumRays,
-	const FVector4* WorldPathDirections,
-	const FVector4* TangentPathDirections,
+	const FVector4f* WorldPathDirections,
+	const FVector4f* TangentPathDirections,
 	EFinalGatherRayBiasMode RayBiasMode,
 	FStaticLightingMappingContext& MappingContext,
 	FLightRay* OutLightRays,
@@ -327,10 +327,10 @@ void FStaticLightingSystem::IntersectLightRays(
 {
 	for (int32 RayIndex = 0; RayIndex < NumRays; RayIndex++)
 	{
-		const FVector4 WorldPathDirection = WorldPathDirections[RayIndex];
-		const FVector4 TangentPathDirection = TangentPathDirections[RayIndex];
+		const FVector4f WorldPathDirection = WorldPathDirections[RayIndex];
+		const FVector4f TangentPathDirection = TangentPathDirections[RayIndex];
 
-		FVector4 SampleOffset(0,0,0);
+		FVector4f SampleOffset(0,0,0);
 		if (GeneralSettings.bAccountForTexelSize)
 		{
 			// Offset the sample's starting point in the tangent XY plane based on the sample's area of influence. 
@@ -340,7 +340,7 @@ void FStaticLightingSystem::IntersectLightRays(
 			
 			// Experiment to distribute the starting position over the area of the texel to anti-alias, causes incorrect shadowing at intersections though
 			//@todo - use consistent sample set between irradiance cache samples
-			//const FVector2D DiskPosition = GetUniformUnitDiskPosition(RandomStream);
+			//const FVector2f DiskPosition = GetUniformUnitDiskPosition(RandomStream);
 			//SampleOffset = Vertex.WorldTangentX * DiskPosition.X * SampleRadius * .5f + Vertex.WorldTangentY * DiskPosition.Y * SampleRadius * .5f;
 		}
 
@@ -351,7 +351,7 @@ void FStaticLightingSystem::IntersectLightRays(
 		// Apply various offsets to the start of the ray.
 		// The offset along the ray direction is to avoid incorrect self-intersection due to floating point precision.
 		// The offset along the normal is to push self-intersection patterns (like triangle shape) on highly curved surfaces onto the backfaces.
-		FVector RayStart = Vertex.WorldPosition
+		FVector3f RayStart = Vertex.WorldPosition
 			+ WorldPathDirection * SceneConstants.VisibilityRayOffsetDistance
 			+ Vertex.WorldTangentZ * RayStartNormalBiasScale * SampleRadius
 			+ SampleOffset;
@@ -383,8 +383,8 @@ void FStaticLightingSystem::IntersectLightRays(
 FLinearColor FStaticLightingSystem::FinalGatherSample(
 	const FStaticLightingMapping* Mapping,
 	const FFullStaticLightingVertex& Vertex,
-	const FVector4& WorldPathDirection,
-	const FVector4& TangentPathDirection,
+	const FVector4f& WorldPathDirection,
+	const FVector4f& TangentPathDirection,
 	const FLightRay& PathRay,
 	const FLightRayIntersection& RayIntersection,
 	float PathSolidAngle,
@@ -397,7 +397,7 @@ FLinearColor FStaticLightingSystem::FinalGatherSample(
 	FLightingCacheGatherInfo& RecordGatherInfo,
 	FFinalGatherInfo& FinalGatherInfo,
 	FFinalGatherHitPoint& HitPoint,
-	FVector& OutUnoccludedSkyVector,
+	FVector3f& OutUnoccludedSkyVector,
 	FLinearColor& OutStationarySkyLighting) const
 {
 	FLinearColor Lighting = FLinearColor::Black;
@@ -405,7 +405,7 @@ FLinearColor FStaticLightingSystem::FinalGatherSample(
 
 	bool bPositiveSample = false;
 
-	OutUnoccludedSkyVector = RayIntersection.bIntersects ? FVector(0) : FVector(WorldPathDirection);
+	OutUnoccludedSkyVector = RayIntersection.bIntersects ? FVector3f(0) : FVector3f(WorldPathDirection);
 
 	if (RayIntersection.bIntersects)
 	{
@@ -522,12 +522,12 @@ class FRefinementTraversalContext
 {
 public:
 	FSimpleQuadTreeNode<FRefinementElement>* Node;
-	FVector2D Min;
-	FVector2D Size;
+	FVector2f Min;
+	FVector2f Size;
 	float SolidAngle;
 	EFinalGatherRefinementCause RefinementCause;
 
-	FRefinementTraversalContext(FSimpleQuadTreeNode<FRefinementElement>* InNode, FVector2D InMin, FVector2D InSize, float InSolidAngle, EFinalGatherRefinementCause InRefinementCause) :
+	FRefinementTraversalContext(FSimpleQuadTreeNode<FRefinementElement>* InNode, FVector2f InMin, FVector2f InSize, float InSolidAngle, EFinalGatherRefinementCause InRefinementCause) :
 		Node(InNode),
 		Min(InMin),
 		Size(InSize),
@@ -536,18 +536,18 @@ public:
 	{}
 };
 
-bool SphereIntersectCone(FSphere SphereCenterAndRadius, FVector ConeVertex, FVector ConeAxis, float ConeAngleCos, float ConeAngleSin)
+bool SphereIntersectCone(FSphere3f SphereCenterAndRadius, FVector3f ConeVertex, FVector3f ConeAxis, float ConeAngleCos, float ConeAngleSin)
 {
-	FVector U = ConeVertex - (SphereCenterAndRadius.W / ConeAngleSin) * ConeAxis;
-	FVector D = SphereCenterAndRadius.Center - U;
-	float DSizeSq = FVector::DotProduct(D, D);
-	float E = FVector::DotProduct(ConeAxis, D);
+	FVector3f U = ConeVertex - (SphereCenterAndRadius.W / ConeAngleSin) * ConeAxis;
+	FVector3f D = SphereCenterAndRadius.Center - U;
+	float DSizeSq = FVector3f::DotProduct(D, D);
+	float E = FVector3f::DotProduct(ConeAxis, D);
 
 	if (E > 0 && E * E >= DSizeSq * ConeAngleCos * ConeAngleCos)
 	{
 		D = SphereCenterAndRadius.Center - ConeVertex;
-		DSizeSq = FVector::DotProduct(D, D);
-		E = -FVector::DotProduct(ConeAxis, D);
+		DSizeSq = FVector3f::DotProduct(D, D);
+		E = -FVector3f::DotProduct(ConeAxis, D);
 
 		if (E > 0 && E * E >= DSizeSq * ConeAngleSin * ConeAngleSin)
 		{
@@ -579,7 +579,7 @@ public:
 	 * Fetches a leaf node value at the desired fractional position.
 	 * Expects a UV that is the center of the cell being searched for, not the min. 
 	 */
-	const FLightingAndOcclusion& GetValue(FVector2D UV)
+	const FLightingAndOcclusion& GetValue(FVector2f UV)
 	{
 		// Theta is radius, clamp
 		const int32 ThetaIndex = FMath::Clamp(FMath::FloorToInt(UV.X * NumThetaSteps), 0, NumThetaSteps - 1);
@@ -632,8 +632,8 @@ public:
 		bool bGatheringForCachedDirectLighting,
 		int32 NumAdaptiveRefinementLevels,
 		float BrightnessThresholdScale,
-		const TArray<FVector4, TInlineAllocator<30> >& TangentImportancePhotonDirections,
-		const TArray<FSphere>& PortalBoundingSpheres,
+		const TArray<FVector4f, TInlineAllocator<30> >& TangentImportancePhotonDirections,
+		const TArray<FSphere3f>& PortalBoundingSpheres,
 		FStaticLightingMappingContext& MappingContext,
 		FGatherHitPoints* HitPointRecorder,
 		FLMRandomStream& RandomStream,
@@ -704,8 +704,8 @@ public:
 				{
 					for (int32 PhiIndex = 0; PhiIndex < NumPhiSteps; PhiIndex++)
 					{
-						const FVector4 CellCenterTangentDirection = UniformSampleHemisphere((ThetaIndex + .5f) / (float)NumThetaSteps, (PhiIndex + .5f) / (float)NumPhiSteps);
-						const FVector4 CellCenterWorldDirection = Vertex.TransformTriangleTangentVectorToWorld(CellCenterTangentDirection);
+						const FVector4f CellCenterTangentDirection = UniformSampleHemisphere((ThetaIndex + .5f) / (float)NumThetaSteps, (PhiIndex + .5f) / (float)NumPhiSteps);
+						const FVector4f CellCenterWorldDirection = Vertex.TransformTriangleTangentVectorToWorld(CellCenterTangentDirection);
 						float IntersectingImportanceConeWeight = 0;
 
 						// Accumulate weight of intersecting photon cones
@@ -796,8 +796,8 @@ public:
 
 							NextNodesToRefine->Add(FRefinementTraversalContext(
 								Node, 
-								FVector2D(ThetaIndex / (float)NumThetaSteps, PhiIndex / (float)NumPhiSteps),
-								FVector2D(1 / (float)NumThetaSteps, 1 / (float)NumPhiSteps),
+								FVector2f(ThetaIndex / (float)NumThetaSteps, PhiIndex / (float)NumPhiSteps),
+								FVector2f(1 / (float)NumThetaSteps, 1 / (float)NumPhiSteps),
 								RootSolidAngle,
 								RefinementCause));
 						}
@@ -818,7 +818,7 @@ public:
 				if (CurrentNodesToRefine->Num() > 0)
 				{
 					FRefinementTraversalContext NodeContext = (*CurrentNodesToRefine)[0];
-					const FVector2D HalfSubCellSize = NodeContext.Size / 4;
+					const FVector2f HalfSubCellSize = NodeContext.Size / 4;
 					// Approximate the cone angle of the sub cell
 					const float SubCellAngle = PI * FMath::Sqrt(HalfSubCellSize.X * HalfSubCellSize.X + HalfSubCellSize.Y * HalfSubCellSize.Y);
 					SubCellCombinedAngleThreshold = FMath::Cos(ImportanceConeAngle + SubCellAngle);
@@ -830,7 +830,7 @@ public:
 				for (int32 NodeIndex = 0; NodeIndex < CurrentNodesToRefine->Num(); NodeIndex++)
 				{
 					FRefinementTraversalContext NodeContext = (*CurrentNodesToRefine)[NodeIndex];
-					const FVector2D HalfSubCellSize = NodeContext.Size / 4;
+					const FVector2f HalfSubCellSize = NodeContext.Size / 4;
 
 					for (int32 SubThetaIndex = 0; SubThetaIndex < NumSubsamples; SubThetaIndex++)
 					{
@@ -838,11 +838,11 @@ public:
 						{
 							FSimpleQuadTreeNode<FRefinementElement>* ChildNode = NodeContext.Node->Children[SubThetaIndex * NumSubsamples + SubPhiIndex];
 
-							const FVector4 CellCenterTangentDirection = UniformSampleHemisphere(
+							const FVector4f CellCenterTangentDirection = UniformSampleHemisphere(
 								NodeContext.Min.X + SubThetaIndex * NodeContext.Size.X / 2 + NodeContext.Size.X / 4,
 								NodeContext.Min.Y + SubPhiIndex * NodeContext.Size.Y / 2 + NodeContext.Size.Y / 4);
 
-							const FVector4 CellCenterWorldDirection = Vertex.TransformTriangleTangentVectorToWorld(CellCenterTangentDirection);
+							const FVector4f CellCenterWorldDirection = Vertex.TransformTriangleTangentVectorToWorld(CellCenterTangentDirection);
 
 							float IntersectingImportanceConeWeight = 0;
 
@@ -900,7 +900,7 @@ public:
 									const float NeighborV = NodeContext.Min.Y + (SubPhiIndex + Neighbors[NeighborIndex].Y) * NodeContext.Size.Y / 2;
 
 									// Query must be done on the center of the cell
-									const FVector2D NeighborUV = FVector2D(NeighborU, NeighborV) + NodeContext.Size / 4;
+									const FVector2f NeighborUV = FVector2f(NeighborU, NeighborV) + NodeContext.Size / 4;
 									const FLightingAndOcclusion NeighborLighting = GetValue(NeighborUV);
 									const float NeighborBrightness = (NeighborLighting.Lighting + NeighborLighting.StationarySkyLighting).GetLuminance();
 									const float NeighborRelativeBrightness = NeighborBrightness / AverageBrightness;
@@ -929,7 +929,7 @@ public:
 							{
 								NextNodesToRefine->Add(FRefinementTraversalContext(
 									ChildNode, 
-									FVector2D(NodeContext.Min.X + SubThetaIndex * NodeContext.Size.X / 2, NodeContext.Min.Y + SubPhiIndex * NodeContext.Size.Y / 2),
+									FVector2f(NodeContext.Min.X + SubThetaIndex * NodeContext.Size.X / 2, NodeContext.Min.Y + SubPhiIndex * NodeContext.Size.Y / 2),
 									NodeContext.Size / 2.0f,
 									SubCellSolidAngle,
 									RefinementCause));
@@ -942,8 +942,8 @@ public:
 			// Swap input and output for the next step
 			Swap(CurrentNodesToRefine, NextNodesToRefine);
 
-			FVector4 WorldPathDirections[4];
-			FVector4 TangentPathDirections[4];
+			FVector4f WorldPathDirections[4];
+			FVector4f TangentPathDirections[4];
 			FLightRay LightRays[4];
 			FLightRayIntersection LightRayIntersections[4];
 
@@ -972,7 +972,7 @@ public:
 							FreeNode = new FSimpleQuadTreeNode<FRefinementElement>();
 						}
 
-						const FVector2D ChildMin = NodeContext.Min + FVector2D(SubThetaIndex, SubPhiIndex) * NodeContext.Size / 2;
+						const FVector2f ChildMin = NodeContext.Min + FVector2f(SubThetaIndex, SubPhiIndex) * NodeContext.Size / 2;
 
 						// Reuse the parent sample result in whatever child cell it falls in
 						if (NodeContext.Node->Element.Uniforms.X >= ChildMin.X
@@ -994,7 +994,7 @@ public:
 							const float Fraction1 = NodeContext.Min.X + SubStepFraction1 * NodeContext.Size.X;
 							const float Fraction2 = NodeContext.Min.Y + SubStepFraction2 * NodeContext.Size.Y;
 
-							const FVector4 SampleDirection = UniformSampleHemisphere(Fraction1, Fraction2);
+							const FVector4f SampleDirection = UniformSampleHemisphere(Fraction1, Fraction2);
 
 							Vertex.ComputePathDirections(SampleDirection, WorldPathDirections[0], TangentPathDirections[0]);
 
@@ -1010,7 +1010,7 @@ public:
 								LightRays,
 								LightRayIntersections);
 
-							FVector UnoccludedSkyVector;
+							FVector3f UnoccludedSkyVector;
 							FLinearColor StationarySkyLighting;
 							FFinalGatherInfo SubsampleFinalGatherInfo;
 							FFinalGatherHitPoint HitPoint;
@@ -1044,7 +1044,7 @@ public:
 								SubsampleGatherInfo.HitPointRecorder->GatherHitPointData.Add(HitPoint);
 							}
 
-							FreeNode->Element = FRefinementElement(FLightingAndOcclusion(SubsampleLighting, UnoccludedSkyVector, StationarySkyLighting, SubsampleFinalGatherInfo.NumSamplesOccluded), FVector2D(Fraction1, Fraction2), StoredHitPointIndex);
+							FreeNode->Element = FRefinementElement(FLightingAndOcclusion(SubsampleLighting, UnoccludedSkyVector, StationarySkyLighting, SubsampleFinalGatherInfo.NumSamplesOccluded), FVector2f(Fraction1, Fraction2), StoredHitPointIndex);
 
 							Stats.NumRefiningFinalGatherSamples[RefinementDepth]++;
 
@@ -1146,10 +1146,10 @@ SampleType FStaticLightingSystem::IncomingRadianceAdaptive(
 	EHemisphereGatherClassification GatherClassification,
 	int32 NumAdaptiveRefinementLevels,
 	float BrightnessThresholdScale,
-	const TArray<FVector4>& UniformHemisphereSamples,
-	const TArray<FVector2D>& UniformHemisphereSampleUniforms,
+	const TArray<FVector4f>& UniformHemisphereSamples,
+	const TArray<FVector2f>& UniformHemisphereSampleUniforms,
 	float MaxUnoccludedLength,
-	const TArray<FVector4>& ImportancePhotonDirections,
+	const TArray<FVector4f>& ImportancePhotonDirections,
 	FStaticLightingMappingContext& MappingContext,
 	FLMRandomStream& RandomStream,
 	FLightingCacheGatherInfo& GatherInfo,
@@ -1174,8 +1174,8 @@ SampleType FStaticLightingSystem::IncomingRadianceAdaptive(
 
 	const float BaseGridSolidAngle = 2 * PI / (float)UniformHemisphereSamples.Num();
 
-	FVector4 WorldPathDirections[4];
-	FVector4 TangentPathDirections[4];
+	FVector4f WorldPathDirections[4];
+	FVector4f TangentPathDirections[4];
 	FLightRay LightRays[4];
 	FLightRayIntersection LightRayIntersections[4];
 
@@ -1185,7 +1185,7 @@ SampleType FStaticLightingSystem::IncomingRadianceAdaptive(
 		for (int32 PhiIndex = 0; PhiIndex < NumPhiSteps; PhiIndex++)
 		{
 			const int32 SampleIndex = ThetaIndex * NumPhiSteps + PhiIndex;
-			const FVector4 TriangleTangentPathDirection = UniformHemisphereSamples[SampleIndex];
+			const FVector4f TriangleTangentPathDirection = UniformHemisphereSamples[SampleIndex];
 			
 			Vertex.ComputePathDirections(TriangleTangentPathDirection, WorldPathDirections[0], TangentPathDirections[0]);
 
@@ -1201,7 +1201,7 @@ SampleType FStaticLightingSystem::IncomingRadianceAdaptive(
 				LightRays,
 				LightRayIntersections);
 
-			FVector UnoccludedSkyVector;
+			FVector3f UnoccludedSkyVector;
 			FLinearColor StationarySkyLighting;
 			FFinalGatherInfo FinalGatherInfo;
 			FFinalGatherHitPoint HitPoint;
@@ -1251,7 +1251,7 @@ SampleType FStaticLightingSystem::IncomingRadianceAdaptive(
 
 	if (bRefine)
 	{
-		TArray<FVector4, TInlineAllocator<30> > TangentSpaceImportancePhotonDirections;
+		TArray<FVector4f, TInlineAllocator<30> > TangentSpaceImportancePhotonDirections;
 		TangentSpaceImportancePhotonDirections.Empty(ImportancePhotonDirections.Num());
 
 		for (int32 PhotonIndex = 0; PhotonIndex < ImportancePhotonDirections.Num(); PhotonIndex++)
@@ -1290,7 +1290,7 @@ SampleType FStaticLightingSystem::IncomingRadianceAdaptive(
 #endif
 
 	SampleType IncomingRadiance;
-	FVector CombinedSkyUnoccludedDirection(0);
+	FVector3f CombinedSkyUnoccludedDirection(0);
 	float NumSamplesOccluded = 0;
 
 	// Accumulate lighting from all samples
@@ -1300,14 +1300,14 @@ SampleType FStaticLightingSystem::IncomingRadianceAdaptive(
 		{
 			const int32 SampleIndex = ThetaIndex * NumPhiSteps + PhiIndex;
 
-			const FVector4 TriangleTangentPathDirection = UniformHemisphereSamples[SampleIndex];
+			const FVector4f TriangleTangentPathDirection = UniformHemisphereSamples[SampleIndex];
 			checkSlow(TriangleTangentPathDirection.Z >= 0.0f);
 			checkSlow(TriangleTangentPathDirection.IsUnit3());
 
-			const FVector4 WorldPathDirection = Vertex.TransformTriangleTangentVectorToWorld(TriangleTangentPathDirection);
+			const FVector4f WorldPathDirection = Vertex.TransformTriangleTangentVectorToWorld(TriangleTangentPathDirection);
 			checkSlow(WorldPathDirection.IsUnit3());
 
-			const FVector4 TangentPathDirection = Vertex.TransformWorldVectorToTangent(WorldPathDirection);
+			const FVector4f TangentPathDirection = Vertex.TransformWorldVectorToTangent(WorldPathDirection);
 			checkSlow(TangentPathDirection.IsUnit3());
 
 			const float UniformPDF = 1.0f / (2.0f * (float)PI);
@@ -1339,7 +1339,7 @@ SampleType FStaticLightingSystem::IncomingRadianceAdaptive(
 	const float NormalizationConstant = .5f * (AmbientOcclusionSettings.OcclusionExponent + 1);
 	IncomingRadiance.SetOcclusion(FMath::Clamp(NormalizationConstant * FMath::Pow(OcclusionFraction, AmbientOcclusionSettings.OcclusionExponent), 0.0f, 1.0f)); 
 
-	const FVector BentNormal = CombinedSkyUnoccludedDirection / (MaxUnoccludedLength * UniformHemisphereSamples.Num());
+	const FVector3f BentNormal = CombinedSkyUnoccludedDirection / (MaxUnoccludedLength * UniformHemisphereSamples.Num());
 	IncomingRadiance.SetSkyOcclusion(BentNormal);
 
 	RefinementGrid.ReturnToFreeList(MappingContext.RefinementTreeFreePool);
@@ -1351,12 +1351,12 @@ SampleType FStaticLightingSystem::IncomingRadianceAdaptive(
 void FStaticLightingSystem::CalculateIrradianceGradients(
 	int32 BounceNumber,
 	const FLightingCacheGatherInfo& GatherInfo,
-	FVector4& RotationalGradient,
-	FVector4& TranslationalGradient) const
+	FVector4f& RotationalGradient,
+	FVector4f& TranslationalGradient) const
 {
 	// Calculate rotational and translational gradients as described in the paper "Irradiance Gradients" by Greg Ward and Paul Heckbert
-	FVector4 AccumulatedRotationalGradient(0,0,0);
-	FVector4 AccumulatedTranslationalGradient(0,0,0);
+	FVector4f AccumulatedRotationalGradient(0,0,0);
+	FVector4f AccumulatedTranslationalGradient(0,0,0);
 	if (IrradianceCachingSettings.bUseIrradianceGradients)
 	{
 		// Extract Theta and Phi steps from the number of hemisphere samples requested
@@ -1369,18 +1369,18 @@ void FStaticLightingSystem::CalculateIrradianceGradients(
 		// Calculate the rotational gradient
 		for (int32 PhiIndex = 0; PhiIndex < NumPhiSteps; PhiIndex++)
 		{
-			FVector4 InnerSum(0,0,0);
+			FVector4f InnerSum(0,0,0);
 			for (int32 ThetaIndex = 0; ThetaIndex < NumThetaSteps; ThetaIndex++)
 			{
 				const int32 SampleIndex = ThetaIndex * NumPhiSteps + PhiIndex;
 				const FLinearColor& IncidentRadiance = GatherInfo.PreviousIncidentRadiances[SampleIndex];
 				// Note: These equations need to be re-derived from the paper for a non-uniform PDF
 				const float TangentTerm = -FMath::Tan(ThetaIndex / (float)NumThetaSteps);
-				InnerSum += TangentTerm * FVector4(IncidentRadiance);
+				InnerSum += TangentTerm * FVector4f(IncidentRadiance);
 			}
 			const float CurrentPhi = 2.0f * (float)PI * PhiIndex / (float)NumPhiSteps;
 			// Vector in the tangent plane perpendicular to the current Phi
-			const FVector4 BasePlaneVector = FVector2D((float)HALF_PI, FMath::Fmod(CurrentPhi + (float)HALF_PI, 2.0f * (float)PI)).SphericalToUnitCartesian();
+			const FVector4f BasePlaneVector = FVector2f((float)HALF_PI, FMath::Fmod(CurrentPhi + (float)HALF_PI, 2.0f * (float)PI)).SphericalToUnitCartesian();
 			AccumulatedRotationalGradient += InnerSum * BasePlaneVector;
 		}
 		// Normalize the sum
@@ -1389,7 +1389,7 @@ void FStaticLightingSystem::CalculateIrradianceGradients(
 		// Calculate the translational gradient
 		for (int32 PhiIndex = 0; PhiIndex < NumPhiSteps; PhiIndex++)
 		{
-			FVector4 PolarWallContribution(0,0,0);
+			FVector4f PolarWallContribution(0,0,0);
 			// Starting from 1 since Theta doesn't wrap around (unlike Phi)
 			for (int32 ThetaIndex = 1; ThetaIndex < NumThetaSteps; ThetaIndex++)
 			{
@@ -1402,13 +1402,13 @@ void FStaticLightingSystem::CalculateIrradianceGradients(
 				const float MinDistance = FMath::Min(PreviousThetaDistance, CurrentThetaDistance);
 				checkSlow(MinDistance > 0);
 				const FLinearColor IncomingRadianceDifference = GatherInfo.PreviousIncidentRadiances[SampleIndex] - GatherInfo.PreviousIncidentRadiances[PreviousThetaSampleIndex];
-				PolarWallContribution += FMath::Sin(CurrentTheta) * CosCurrentTheta * CosCurrentTheta / MinDistance * FVector4(IncomingRadianceDifference);
+				PolarWallContribution += FMath::Sin(CurrentTheta) * CosCurrentTheta * CosCurrentTheta / MinDistance * FVector4f(IncomingRadianceDifference);
 				checkSlow(!PolarWallContribution.ContainsNaN());
 			}
 
 			// Wrap Phi around for the first Phi index
 			const int32 PreviousPhiIndex = PhiIndex == 0 ? NumPhiSteps - 1 : PhiIndex - 1;
-			FVector4 RadialWallContribution(0,0,0);
+			FVector4f RadialWallContribution(0,0,0);
 			for (int32 ThetaIndex = 0; ThetaIndex < NumThetaSteps; ThetaIndex++)
 			{
 				const float CurrentTheta = FMath::Acos(ThetaIndex / (float)NumThetaSteps);
@@ -1420,15 +1420,15 @@ void FStaticLightingSystem::CalculateIrradianceGradients(
 				const float MinDistance = FMath::Min(PreviousPhiDistance, CurrentPhiDistance);
 				checkSlow(MinDistance > 0);
 				const FLinearColor IncomingRadianceDifference = GatherInfo.PreviousIncidentRadiances[SampleIndex] - GatherInfo.PreviousIncidentRadiances[PreviousPhiSampleIndex];
-				RadialWallContribution += (FMath::Sin(NextTheta) - FMath::Sin(CurrentTheta)) / MinDistance * FVector4(IncomingRadianceDifference);
+				RadialWallContribution += (FMath::Sin(NextTheta) - FMath::Sin(CurrentTheta)) / MinDistance * FVector4f(IncomingRadianceDifference);
 				checkSlow(!RadialWallContribution.ContainsNaN());
 			}
 
 			const float CurrentPhi = 2.0f * (float)PI * PhiIndex / (float)NumPhiSteps;
 			// Vector in the tangent plane in the direction of the current Phi
-			const FVector4 PhiDirection = SphericalToUnitCartesian(FVector2D((float)HALF_PI, CurrentPhi));
+			const FVector4f PhiDirection = SphericalToUnitCartesian(FVector2f((float)HALF_PI, CurrentPhi));
 			// Vector in the tangent plane perpendicular to the current Phi
-			const FVector4 PerpendicularPhiDirection = FVector2D((float)HALF_PI, FMath::Fmod(CurrentPhi + (float)HALF_PI, 2.0f * (float)PI)).SphericalToUnitCartesian();
+			const FVector4f PerpendicularPhiDirection = FVector2f((float)HALF_PI, FMath::Fmod(CurrentPhi + (float)HALF_PI, 2.0f * (float)PI)).SphericalToUnitCartesian();
 
 			PolarWallContribution = PhiDirection * 2.0f * (float)PI / (float)NumPhiSteps * PolarWallContribution;
 			RadialWallContribution = PerpendicularPhiDirection * RadialWallContribution;
@@ -1553,7 +1553,7 @@ FFinalGatherSample FStaticLightingSystem::CachePointIncomingRadiance(
 		else
 		{
 			// Using final gathering with photon mapping, hemisphere gathering without photon mapping, path tracing and/or just calculating ambient occlusion
-			TArray<FVector4> ImportancePhotonDirections;
+			TArray<FVector4f> ImportancePhotonDirections;
 
 			if (GeneralSettings.NumIndirectLightingBounces > 0)
 			{
@@ -1585,7 +1585,7 @@ FFinalGatherSample FStaticLightingSystem::CachePointIncomingRadiance(
 						const FPhoton& CurrentPhoton = FoundPhotons[PhotonIndex];
 						// Calculate the direction from the current position to the photon's source
 						// Using the photon's incident direction unmodified produces artifacts proportional to the distance to that photon
-						const FVector4 NewDirection = CurrentPhoton.GetPosition() + CurrentPhoton.GetIncidentDirection() * CurrentPhoton.GetDistance() - Vertex.WorldPosition;
+						const FVector4f NewDirection = CurrentPhoton.GetPosition() + CurrentPhoton.GetIncidentDirection() * CurrentPhoton.GetDistance() - Vertex.WorldPosition;
 						// Only use the direction if it is in the hemisphere of the normal
 						// FindNearbyPhotons only returns photons whose incident directions lie in this hemisphere, but the recalculated direction might not.
 						if (Dot3(NewDirection, Vertex.TriangleNormal) > 0.0f)
@@ -1624,8 +1624,8 @@ FFinalGatherSample FStaticLightingSystem::CachePointIncomingRadiance(
 
 			if (IrradianceCachingSettings.bAllowIrradianceCaching)
 			{
-				FVector4 RotationalGradient;
-				FVector4 TranslationalGradient;
+				FVector4f RotationalGradient;
+				FVector4f TranslationalGradient;
 				CalculateIrradianceGradients(BounceNumber, GatherInfo, RotationalGradient, TranslationalGradient);
 
 				float OverrideRadius = 0;
@@ -1646,10 +1646,10 @@ FFinalGatherSample FStaticLightingSystem::CachePointIncomingRadiance(
 						// Trace a ray from one corner of the texel to the other to detect this edge case, 
 						// And set the record radius to the minimum to contain the error.
 						// Center of the texel offset along the normal
-						const FVector4 TexelCenterOffset = Vertex.WorldPosition + Vertex.TriangleNormal * SampleRadius * SceneConstants.VisibilityNormalOffsetSampleRadiusScale;
+						const FVector4f TexelCenterOffset = Vertex.WorldPosition + Vertex.TriangleNormal * SampleRadius * SceneConstants.VisibilityNormalOffsetSampleRadiusScale;
 						// Vector from the center to one of the corners of the texel
 						// The FMath::Sqrt(.5f) is to normalize (Vertex.TriangleTangentX + Vertex.TriangleTangentY), which are orthogonal unit vectors.
-						const FVector4 CornerOffset = FMath::Sqrt(.5f) * (Vertex.TriangleTangentX + Vertex.TriangleTangentY) * SampleRadius * SceneConstants.VisibilityTangentOffsetSampleRadiusScale;
+						const FVector4f CornerOffset = FMath::Sqrt(.5f) * (Vertex.TriangleTangentX + Vertex.TriangleTangentY) * SampleRadius * SceneConstants.VisibilityTangentOffsetSampleRadiusScale;
 						const FLightRay TexelRay(
 							TexelCenterOffset + CornerOffset,
 							TexelCenterOffset - CornerOffset,
@@ -1708,7 +1708,7 @@ FFinalGatherSample FStaticLightingSystem::CachePointIncomingRadiance(
 #if ALLOW_LIGHTMAP_SAMPLE_DEBUGGING
 				if (IrradianceCachingSettings.bVisualizeIrradianceSamples && Mapping == Scene.DebugMapping && GeneralSettings.ViewSingleBounceNumber == BounceNumber)
 				{
-					const float DistanceToDebugTexelSq = FVector(Scene.DebugInput.Position - Vertex.WorldPosition).SizeSquared();
+					const float DistanceToDebugTexelSq = FVector3f(Scene.DebugInput.Position - Vertex.WorldPosition).SizeSquared();
 					FDebugLightingCacheRecord TempRecord;
 					TempRecord.bNearSelectedTexel = DistanceToDebugTexelSq < NewRecord.BoundingRadius * NewRecord.BoundingRadius;
 					TempRecord.Radius = NewRecord.Radius;
diff --git a/Engine/Source/Programs/UnrealLightmass/Private/Lighting/FluidSurface.cpp b/Engine/Source/Programs/UnrealLightmass/Private/Lighting/FluidSurface.cpp
--- a/Engine/Source/Programs/UnrealLightmass/Private/Lighting/FluidSurface.cpp
+++ b/Engine/Source/Programs/UnrealLightmass/Private/Lighting/FluidSurface.cpp
@@ -9,22 +9,22 @@ namespace Lightmass
 
 	/** Represents the fluid surface mesh to the static lighting system. */
 	static void GetStaticLightingVertex(
-		const FVector4* QuadCorners,
-		const FVector4* QuadUVCorners,
+		const FVector4f* QuadCorners,
+		const FVector4f* QuadUVCorners,
 		uint32 VertexIndex,
-		const FMatrix& LocalToWorld,
-		const FMatrix& LocalToWorldInverseTranspose,
+		const FMatrix44f& LocalToWorld,
+		const FMatrix44f& LocalToWorldInverseTranspose,
 		FStaticLightingVertex& OutVertex
 		)
 	{
 		OutVertex.WorldPosition = LocalToWorld.TransformPosition(QuadCorners[VertexIndex]);
-		OutVertex.WorldTangentX = LocalToWorld.TransformVector(FVector4(1, 0, 0, 1)).GetSafeNormal();
-		OutVertex.WorldTangentY = LocalToWorld.TransformVector(FVector4(0, 1, 0, 1)).GetSafeNormal();
-		OutVertex.WorldTangentZ = LocalToWorldInverseTranspose.TransformVector(FVector4(0, 0, 1, 1)).GetSafeNormal();
+		OutVertex.WorldTangentX = LocalToWorld.TransformVector(FVector4f(1, 0, 0, 1)).GetSafeNormal();
+		OutVertex.WorldTangentY = LocalToWorld.TransformVector(FVector4f(0, 1, 0, 1)).GetSafeNormal();
+		OutVertex.WorldTangentZ = LocalToWorldInverseTranspose.TransformVector(FVector4f(0, 0, 1, 1)).GetSafeNormal();
 
 		for(uint32 UVIndex = 0; UVIndex < 1; UVIndex++)
 		{
-			OutVertex.TextureCoordinates[UVIndex] = FVector2D(QuadUVCorners[VertexIndex].X, QuadUVCorners[VertexIndex].Y);
+			OutVertex.TextureCoordinates[UVIndex] = FVector2f(QuadUVCorners[VertexIndex].X, QuadUVCorners[VertexIndex].Y);
 		}
 	}
 
diff --git a/Engine/Source/Programs/UnrealLightmass/Private/Lighting/GatheredLightingSample.h b/Engine/Source/Programs/UnrealLightmass/Private/Lighting/GatheredLightingSample.h
--- a/Engine/Source/Programs/UnrealLightmass/Private/Lighting/GatheredLightingSample.h
+++ b/Engine/Source/Programs/UnrealLightmass/Private/Lighting/GatheredLightingSample.h
@@ -26,7 +26,7 @@ public:
 	float SHCorrection;
 
 	/** Sky bent normal, points toward the most unoccluded direction, and the length is the visibility amount (0 = occluded, 1 = visible). */
-	FVector SkyOcclusion;
+	FVector3f SkyOcclusion;
 
 	float AOMaterialMask;
 
@@ -35,7 +35,7 @@ public:
 	{
 		SHCorrection = 0.0f;
 		IncidentLighting = FLinearColor(0, 0, 0, 0);
-		SkyOcclusion = FVector(0);
+		SkyOcclusion = FVector3f(0);
 		AOMaterialMask = 0;
 	}
 
@@ -43,7 +43,7 @@ public:
 	{
 		SHCorrection = 0.0f;
 		IncidentLighting = FLinearColor(0, 0, 0, 0);
-		SkyOcclusion = FVector(0);
+		SkyOcclusion = FVector3f(0);
 		AOMaterialMask = 0;
 	}	
 
@@ -56,7 +56,7 @@ public:
 
 	void ApplyOcclusion(float Occlusion);
 
-	inline void SetSkyOcclusion(FVector InSkyOcclusion)
+	inline void SetSkyOcclusion(FVector3f InSkyOcclusion)
 	{
 		SkyOcclusion = InSkyOcclusion;
 	}
@@ -104,7 +104,7 @@ public:
 	* @param Direction - The direction toward the light at the sample point.
 	*/
 	template <int32 SHOrder>
-	static TGatheredLightSample<SHOrder> PointLightWorldSpace(const FLinearColor& Color, const FVector4& TangentDirection, const FVector4& WorldDirection);
+	static TGatheredLightSample<SHOrder> PointLightWorldSpace(const FLinearColor& Color, const FVector4f& TangentDirection, const FVector4f& WorldDirection);
 };
 
 typedef TGatheredLightSample<2> FGatheredLightSample;
@@ -205,12 +205,12 @@ public:
 		Occlusion = InOcclusion;
 	}
 
-	inline void AddIncomingRadiance(const FLinearColor& IncomingRadiance, float Weight, const FVector4& TangentSpaceDirection, const FVector4& WorldSpaceDirection)
+	inline void AddIncomingRadiance(const FLinearColor& IncomingRadiance, float Weight, const FVector4f& TangentSpaceDirection, const FVector4f& WorldSpaceDirection)
 	{
 		AddWeighted(FGatheredLightSampleUtil::PointLightWorldSpace<SHOrder>(IncomingRadiance, TangentSpaceDirection, WorldSpaceDirection), Weight);
 	}
 
-	inline void AddIncomingStationarySkyLight(const FLinearColor& IncomingSkyLight, float Weight, const FVector4& TangentSpaceDirection, const FVector4& WorldSpaceDirection)
+	inline void AddIncomingStationarySkyLight(const FLinearColor& IncomingSkyLight, float Weight, const FVector4f& TangentSpaceDirection, const FVector4f& WorldSpaceDirection)
 	{
 		StationarySkyLighting.AddWeighted(FGatheredLightSampleUtil::PointLightWorldSpace<SHOrder>(IncomingSkyLight, TangentSpaceDirection, WorldSpaceDirection), Weight);
 	}
diff --git a/Engine/Source/Programs/UnrealLightmass/Private/Lighting/Landscape.cpp b/Engine/Source/Programs/UnrealLightmass/Private/Lighting/Landscape.cpp
--- a/Engine/Source/Programs/UnrealLightmass/Private/Lighting/Landscape.cpp
+++ b/Engine/Source/Programs/UnrealLightmass/Private/Lighting/Landscape.cpp
@@ -62,22 +62,22 @@ namespace Lightmass
 		OutVertex.WorldTangentZ.X = 2.0f / 255.f * (float)Data->B - 1.0f;
 		OutVertex.WorldTangentZ.Y = 2.0f / 255.f * (float)Data->A - 1.0f;
 		OutVertex.WorldTangentZ.Z = FMath::Sqrt(FMath::Max(1.0f - (FMath::Square(OutVertex.WorldTangentZ.X) + FMath::Square(OutVertex.WorldTangentZ.Y)), 0.f));
-		OutVertex.WorldTangentX = FVector4(OutVertex.WorldTangentZ.Z, 0.0f, -OutVertex.WorldTangentZ.X);
+		OutVertex.WorldTangentX = FVector4f(OutVertex.WorldTangentZ.Z, 0.0f, -OutVertex.WorldTangentZ.X);
 		OutVertex.WorldTangentY = OutVertex.WorldTangentZ ^ OutVertex.WorldTangentX;
 
 		// Copied (vaguely) from FLandscapeComponentDataInterface::GetWorldPositionTangents to fix bad lighting when rotated
-		const FMatrix LtWNoScale = LocalToWorld.GetMatrixWithoutScale();
+		const FMatrix44f LtWNoScale = LocalToWorld.GetMatrixWithoutScale();
 		OutVertex.WorldTangentX = LtWNoScale.TransformVector(OutVertex.WorldTangentX);
 		OutVertex.WorldTangentY = LtWNoScale.TransformVector(OutVertex.WorldTangentY);
 		OutVertex.WorldTangentZ = LtWNoScale.TransformVector(OutVertex.WorldTangentZ);
 
 		const uint16 Height = (Data->R << 8) + Data->G;
-		OutVertex.WorldPosition = LocalToWorld.TransformPosition( FVector4( LocalX, LocalY, ((float)Height - 32768.f) * LANDSCAPE_ZSCALE ) );
+		OutVertex.WorldPosition = LocalToWorld.TransformPosition( FVector4f( LocalX, LocalY, ((float)Height - 32768.f) * LANDSCAPE_ZSCALE ) );
 		//UE_LOG(LogLightmass, Log, TEXT("%d, %d, %d, %d, %d, %d, X:%f, Y:%f, Z:%f "), SectionBaseX + LocalX - ExpandQuadsX, SectionBaseY + LocalY - ExpandQuadsY, ClampedLocalX, ClampedLocalY, SectionBaseX, SectionBaseY, WorldPos.X, WorldPos.Y, WorldPos.Z);
 
 		const int32 LightmapUVIndex = 1;
 
-		OutVertex.TextureCoordinates[0] = FVector2D((float)X / NumVertices, (float)Y / NumVertices); 
+		OutVertex.TextureCoordinates[0] = FVector2f((float)X / NumVertices, (float)Y / NumVertices); 
 		OutVertex.TextureCoordinates[LightmapUVIndex].X = X * UVFactor;
 		OutVertex.TextureCoordinates[LightmapUVIndex].Y = Y * UVFactor;
 	}
diff --git a/Engine/Source/Programs/UnrealLightmass/Private/Lighting/LightingCache.h b/Engine/Source/Programs/UnrealLightmass/Private/Lighting/LightingCache.h
--- a/Engine/Source/Programs/UnrealLightmass/Private/Lighting/LightingCache.h
+++ b/Engine/Source/Programs/UnrealLightmass/Private/Lighting/LightingCache.h
@@ -52,7 +52,7 @@ public:
 		QuantizedPhi(0)
 	{}
 
-	FQuantizedHemisphereDirection(const FVector4& UnitTangentSpaceDirection)
+	FQuantizedHemisphereDirection(const FVector4f& UnitTangentSpaceDirection)
 	{
 		checkSlow(UnitTangentSpaceDirection.IsUnit3());
 		checkSlow(UnitTangentSpaceDirection.Z >= 0 && UnitTangentSpaceDirection.Z <= 1);
@@ -65,14 +65,14 @@ public:
 		QuantizedPhi = (uint8)FMath::Clamp<int32>(FMath::RoundToInt(NormalizedPhi * (MAX_uint8 - 1)), 0, MAX_uint8 - 1);
 	}
 
-	FVector4 GetDirection() const
+	FVector4f GetDirection() const
 	{
 		const float Uint8RangeScale = 1.0f / (float)(MAX_uint8 - 1);
 		const float RescaledTheta = QuantizedTheta * Uint8RangeScale * PI;
 		const float RescaledPhi = (QuantizedPhi * Uint8RangeScale - .5f) * 2 * PI;
 
 		const float SinTheta = FMath::Sin(RescaledTheta);
-		return FVector(SinTheta * FMath::Cos(RescaledPhi), SinTheta * FMath::Sin(RescaledPhi), FMath::Cos(RescaledTheta));
+		return FVector3f(SinTheta * FMath::Cos(RescaledPhi), SinTheta * FMath::Sin(RescaledPhi), FMath::Cos(RescaledTheta));
 	}
 
 private:
@@ -172,10 +172,10 @@ public:
 		RecordSampleType Lighting;
 
 		/** The rotational gradient along the vector perpendicular to both the record normal and the normal of the vertex being interpolated to, used for higher order interpolation. */
-		FVector4 RotationalGradient;
+		FVector4f RotationalGradient;
 
 		/** The translational gradient from the record to the point being interpolated to, used for higher order interpolation. */
-		FVector4 TranslationalGradient;
+		FVector4f TranslationalGradient;
 
 		/** For debugging */
 		int32 Id;
@@ -183,7 +183,7 @@ public:
 		float BackfacingHitsFraction;
 
 		/** Initialization constructor. */
-		FRecord(const FFullStaticLightingVertex& InVertex,int32 InElementIndex,const FLightingCacheGatherInfo& GatherInfo,float SampleRadius,float InOverrideRadius,const FIrradianceCachingSettings& IrradianceCachingSettings,const FStaticLightingSettings& GeneralSettings,const RecordSampleType& InLighting, const FVector4& InRotGradient, const FVector4& InTransGradient):
+		FRecord(const FFullStaticLightingVertex& InVertex,int32 InElementIndex,const FLightingCacheGatherInfo& GatherInfo,float SampleRadius,float InOverrideRadius,const FIrradianceCachingSettings& IrradianceCachingSettings,const FStaticLightingSettings& GeneralSettings,const RecordSampleType& InLighting, const FVector4f& InRotGradient, const FVector4f& InTransGradient):
 			Vertex(InVertex),
 			ElementIndex(InElementIndex),
 			Lighting(InLighting),
@@ -220,12 +220,12 @@ public:
 		{
 			return FBoxCenterAndExtent(
 				LightingRecord.Vertex.WorldPosition,
-				FVector4(LightingRecord.BoundingRadius, LightingRecord.BoundingRadius, LightingRecord.BoundingRadius)
+				FVector4f(LightingRecord.BoundingRadius, LightingRecord.BoundingRadius, LightingRecord.BoundingRadius)
 				);
 		}
 	};
 
-	TLightingCache(const FBox& InBoundingBox, const FStaticLightingSystem& System, int32 InBounceNumber) :
+	TLightingCache(const FBox3f& InBoundingBox, const FStaticLightingSystem& System, int32 InBounceNumber) :
 		FLightingCacheBase(System, InBounceNumber),
 		Octree(InBoundingBox.GetCenter(),InBoundingBox.GetExtent().GetMax())
 	{}
@@ -383,7 +383,7 @@ bool TLightingCache<SampleType>::InterpolateLighting(
 	// Iterate over the octree nodes containing the query point.
 	for( typename LightingOctreeType::template TConstElementBoxIterator<> OctreeIt(
 		Octree,
-		FBoxCenterAndExtent(Vertex.WorldPosition, FVector4(0,0,0))
+		FBoxCenterAndExtent(Vertex.WorldPosition, FVector4f(0,0,0))
 		);
 		OctreeIt.HasPendingElements();
 		OctreeIt.Advance())
@@ -401,9 +401,9 @@ bool TLightingCache<SampleType>::InterpolateLighting(
 
 		// Don't use a lighting record if it's in front of the query point.
 		// Query points behind the lighting record may have nearby occluders that the lighting record does not see.
-		const FVector4 RecordToVertexVector = Vertex.WorldPosition - LightingRecord.Vertex.WorldPosition;
+		const FVector4f RecordToVertexVector = Vertex.WorldPosition - LightingRecord.Vertex.WorldPosition;
 		// Use the average normal to handle surfaces with constant concavity
-		const FVector4 AverageNormal = (LightingRecord.Vertex.TriangleNormal + Vertex.TriangleNormal).GetSafeNormal();
+		const FVector4f AverageNormal = (LightingRecord.Vertex.TriangleNormal + Vertex.TriangleNormal).GetSafeNormal();
 		const float PlaneDistance = Dot3(AverageNormal, RecordToVertexVector.GetSafeNormal());
 		// Setup an error metric that goes from 0 if the points are coplanar, to 1 if the point being shaded is at the angle corresponding to MinCosPointBehindPlane behind the plane
 		const float PointBehindPlaneError = FMath::Max(PlaneDistance / MinCosPointBehindPlane, 0.0f);
diff --git a/Engine/Source/Programs/UnrealLightmass/Private/Lighting/LightingMesh.cpp b/Engine/Source/Programs/UnrealLightmass/Private/Lighting/LightingMesh.cpp
--- a/Engine/Source/Programs/UnrealLightmass/Private/Lighting/LightingMesh.cpp
+++ b/Engine/Source/Programs/UnrealLightmass/Private/Lighting/LightingMesh.cpp
@@ -21,13 +21,13 @@ bool bAllowMeshAreaLights = false;
 FLinearColor FStaticLightingMesh::EvaluateBRDF(
 	const FStaticLightingVertex& Vertex, 
 	int32 ElementIndex,
-	const FVector4& IncomingDirection, 
-	const FVector4& OutgoingDirection) const
+	const FVector4f& IncomingDirection, 
+	const FVector4f& OutgoingDirection) const
 {
 	checkSlow(Vertex.WorldTangentZ.IsUnit3());
 	checkSlow(IncomingDirection.IsUnit3());
 	checkSlow(OutgoingDirection.IsUnit3());
-	const FVector4 ReflectedIncomingVector = IncomingDirection.Reflect3(Vertex.WorldTangentZ);
+	const FVector4f ReflectedIncomingVector = IncomingDirection.Reflect3(Vertex.WorldTangentZ);
 	const float OutgoingDotReflected = FMath::Max(Dot3(OutgoingDirection, ReflectedIncomingVector), 0.0f);
 	const FLinearColor Diffuse = EvaluateDiffuse(Vertex.TextureCoordinates[0], ElementIndex);
 	return Diffuse / (float)PI;
@@ -37,8 +37,8 @@ FLinearColor FStaticLightingMesh::EvaluateBRDF(
 FLinearColor FStaticLightingMesh::SampleBRDF(
 	const FStaticLightingVertex& Vertex, 
 	int32 ElementIndex,
-	const FVector4& IncomingDirection, 
-	FVector4& OutgoingDirection,
+	const FVector4f& IncomingDirection, 
+	FVector4f& OutgoingDirection,
 	float& DirectionPDF,
 	FLMRandomStream& RandomStream
 	) const
@@ -46,15 +46,15 @@ FLinearColor FStaticLightingMesh::SampleBRDF(
 	checkSlow(Vertex.WorldTangentZ.IsUnit3());
 	checkSlow(IncomingDirection.IsUnit3());
 
-	const FVector4 ReflectedIncomingVector = IncomingDirection.Reflect3(Vertex.WorldTangentZ);
-	const FVector4 TangentReflectedIncomingVector = Vertex.TransformWorldVectorToTangent(ReflectedIncomingVector);
+	const FVector4f ReflectedIncomingVector = IncomingDirection.Reflect3(Vertex.WorldTangentZ);
+	const FVector4f TangentReflectedIncomingVector = Vertex.TransformWorldVectorToTangent(ReflectedIncomingVector);
 
 	const FLinearColor Diffuse = EvaluateDiffuse(Vertex.TextureCoordinates[0], ElementIndex);
 
 	const float DiffuseIntensity = Diffuse.GetLuminance();
 	
 	// Generate a direction based on the cosine lobe
-	FVector4 TangentPathDirection = GetCosineHemisphereVector(RandomStream);
+	FVector4f TangentPathDirection = GetCosineHemisphereVector(RandomStream);
 
 	const float CosTheta = FMath::Max(Dot3(IncomingDirection, Vertex.WorldTangentZ), 0.0f);
 	const float CosPDF = CosTheta / (float)PI;
@@ -160,13 +160,13 @@ void FStaticLightingMesh::Import( FLightmassImporter& Importer )
 
 /** Determines whether two triangles overlap each other's AABB's. */
 static bool AxisAlignedTriangleIntersectTriangle2d(
-	const FVector2D& V0, const FVector2D& V1, const FVector2D& V2, 
-	const FVector2D& OtherV0, const FVector2D& OtherV1, const FVector2D& OtherV2)
+	const FVector2f& V0, const FVector2f& V1, const FVector2f& V2, 
+	const FVector2f& OtherV0, const FVector2f& OtherV1, const FVector2f& OtherV2)
 {
-	const FVector2D MinFirst = FMath::Min(V0, FMath::Min(V1, V2));
-	const FVector2D MaxFirst = FMath::Max(V0, FMath::Max(V1, V2));
-	const FVector2D MinSecond = FMath::Min(OtherV0, FMath::Min(OtherV1, OtherV2));
-	const FVector2D MaxSecond = FMath::Max(OtherV0, FMath::Max(OtherV1, OtherV2));
+	const FVector2f MinFirst = FMath::Min(V0, FMath::Min(V1, V2));
+	const FVector2f MaxFirst = FMath::Max(V0, FMath::Max(V1, V2));
+	const FVector2f MinSecond = FMath::Min(OtherV0, FMath::Min(OtherV1, OtherV2));
+	const FVector2f MaxSecond = FMath::Max(OtherV0, FMath::Max(OtherV1, OtherV2));
 
 	return !(MinFirst.X > MaxSecond.X 
 		|| MinSecond.X > MaxFirst.X 
@@ -234,7 +234,7 @@ void FStaticLightingMesh::CreateMeshAreaLights(
 	CalculateUniqueLayers(MeshVertices, ElementIndices, LayeredGroupTriangles);
 
 	// get Min/MaxUV on the mesh for the triangles
-	FVector2D MinUV(FLT_MAX, FLT_MAX), MaxUV(-FLT_MAX, -FLT_MAX);
+	FVector2f MinUV(FLT_MAX, FLT_MAX), MaxUV(-FLT_MAX, -FLT_MAX);
 	for (int32 TriangleIndex = 0; TriangleIndex < NumTriangles; TriangleIndex++)
 	{
 		for (int32 VIndex = 0; VIndex < 3; VIndex++)
@@ -258,8 +258,8 @@ void FStaticLightingMesh::CreateMeshAreaLights(
 
 	// calculate the bias and scale needed to map the random UV range into 0 .. NumIterations when rasterizing
 	// into the TexelToCornersMap
-	const FVector2D UVBias(-FMath::FloorToFloat(MinUV.X), -FMath::FloorToFloat(MinUV.Y));
-	const FVector2D UVScale(1.0f / NumIterationsX, 1.0f / NumIterationsY);
+	const FVector2f UVBias(-FMath::FloorToFloat(MinUV.X), -FMath::FloorToFloat(MinUV.Y));
+	const FVector2f UVScale(1.0f / NumIterationsX, 1.0f / NumIterationsY);
 
 	for (int32 MaterialIndex = 0; MaterialIndex < MaterialElements.Num(); MaterialIndex++)
 	{
@@ -281,7 +281,7 @@ void FStaticLightingMesh::CreateMeshAreaLights(
 					{
 						FTexelToCorners& CurrentTexelCorners = TexelToCornersMap(X, Y);
 						// Normals need to be unit as their dot product will be used in comparisons later
-						CurrentTexelCorners.WorldTangentZ = CurrentTexelCorners.WorldTangentZ.SizeSquared3() > DELTA ? CurrentTexelCorners.WorldTangentZ.GetUnsafeNormal3() : FVector4(0,0,1);
+						CurrentTexelCorners.WorldTangentZ = CurrentTexelCorners.WorldTangentZ.SizeSquared3() > DELTA ? CurrentTexelCorners.WorldTangentZ.GetUnsafeNormal3() : FVector4f(0,0,1);
 					}
 				}
 
@@ -328,7 +328,7 @@ void FStaticLightingMesh::CreateMeshAreaLights(
 				PrimitiveIndices.AddZeroed(TexelToCornersMap.GetSizeX() * TexelToCornersMap.GetSizeY());
 				FMemory::Memset(PrimitiveIndices.GetData(), UnprocessedIndex, PrimitiveIndices.Num() * PrimitiveIndices.GetTypeSize());
 				int32 NextPrimitiveIndex = 0;
-				const float DistanceThreshold = FBoxSphereBounds(BoundingBox).SphereRadius * Scene.MeshAreaLightSettings.MeshAreaLightSimplifyMeshBoundingRadiusFractionThreshold;
+				const float DistanceThreshold = FBoxSphereBounds3f(BoundingBox).SphereRadius * Scene.MeshAreaLightSettings.MeshAreaLightSimplifyMeshBoundingRadiusFractionThreshold;
 				// The temporary stack of texels that need to be processed
 				TArray<FIntPoint> PendingTexels;
 				// Iterate over all texels and assign a primitive index to each one
@@ -343,7 +343,7 @@ void FStaticLightingMesh::CreateMeshAreaLights(
 						checkSlow(LightIndex != PendingProcessingIndex);
 						const FTexelToCorners& CurrentTexelCorners = TexelToCornersMap(X, Y);
 
-						FVector4 PrimitiveCenter(0,0,0);
+						FVector4f PrimitiveCenter(0,0,0);
 						for (int32 CornerIndex = 0; CornerIndex < NumTexelCorners; CornerIndex++)
 						{
 							PrimitiveCenter += CurrentTexelCorners.Corners[CornerIndex].WorldPosition / (float)NumTexelCorners;
@@ -408,7 +408,7 @@ void FStaticLightingMesh::CreateMeshAreaLights(
 							if (PrimitiveIndex >= 0)
 							{
 								// Calculate the texel's center
-								FVector4 TexelCenter(0,0,0);
+								FVector4f TexelCenter(0,0,0);
 								bool bAllCornersValid = true;
 								for (int32 CornerIndex = 0; CornerIndex < NumTexelCorners; CornerIndex++)
 								{
@@ -429,7 +429,7 @@ void FStaticLightingMesh::CreateMeshAreaLights(
 								}
 
 								const float XFraction = X / (float)CurrentMaterial.EmissiveSize;
-								const FLinearColor CurrentEmissive = EvaluateEmissive(FVector2D(XFraction, YFraction), MaterialIndex);
+								const FLinearColor CurrentEmissive = EvaluateEmissive(FVector2f(XFraction, YFraction), MaterialIndex);
 								checkSlow(CurrentEmissive.R > Scene.MeshAreaLightSettings.EmissiveIntensityThreshold 
 									|| CurrentEmissive.G > Scene.MeshAreaLightSettings.EmissiveIntensityThreshold 
 									|| CurrentEmissive.B > Scene.MeshAreaLightSettings.EmissiveIntensityThreshold);
@@ -465,10 +465,10 @@ void FStaticLightingMesh::CreateMeshAreaLights(
 							{
 								// Draw 4 lines between the primitive corners for debugging
 								// Currently hijacking ShadowRays
-								//LightingSystem.DebugOutput.ShadowRays.Add(FDebugStaticLightingRay(EmissivePrimitives[LightIndex][PrimitiveIndex].Corners[0].WorldPosition - FVector4(0,0,.1f), EmissivePrimitives[LightIndex][PrimitiveIndex].Corners[1].WorldPosition - FVector4(0,0,.1f), true, false));
-								//LightingSystem.DebugOutput.ShadowRays.Add(FDebugStaticLightingRay(EmissivePrimitives[LightIndex][PrimitiveIndex].Corners[1].WorldPosition - FVector4(0,0,.1f), EmissivePrimitives[LightIndex][PrimitiveIndex].Corners[3].WorldPosition - FVector4(0,0,.1f), true, true));
-								//LightingSystem.DebugOutput.ShadowRays.Add(FDebugStaticLightingRay(EmissivePrimitives[LightIndex][PrimitiveIndex].Corners[3].WorldPosition - FVector4(0,0,.1f), EmissivePrimitives[LightIndex][PrimitiveIndex].Corners[2].WorldPosition - FVector4(0,0,.1f), true, false));
-								//LightingSystem.DebugOutput.ShadowRays.Add(FDebugStaticLightingRay(EmissivePrimitives[LightIndex][PrimitiveIndex].Corners[2].WorldPosition - FVector4(0,0,.1f), EmissivePrimitives[LightIndex][PrimitiveIndex].Corners[0].WorldPosition - FVector4(0,0,.1f), true, true));
+								//LightingSystem.DebugOutput.ShadowRays.Add(FDebugStaticLightingRay(EmissivePrimitives[LightIndex][PrimitiveIndex].Corners[0].WorldPosition - FVector4f(0,0,.1f), EmissivePrimitives[LightIndex][PrimitiveIndex].Corners[1].WorldPosition - FVector4f(0,0,.1f), true, false));
+								//LightingSystem.DebugOutput.ShadowRays.Add(FDebugStaticLightingRay(EmissivePrimitives[LightIndex][PrimitiveIndex].Corners[1].WorldPosition - FVector4f(0,0,.1f), EmissivePrimitives[LightIndex][PrimitiveIndex].Corners[3].WorldPosition - FVector4f(0,0,.1f), true, true));
+								//LightingSystem.DebugOutput.ShadowRays.Add(FDebugStaticLightingRay(EmissivePrimitives[LightIndex][PrimitiveIndex].Corners[3].WorldPosition - FVector4f(0,0,.1f), EmissivePrimitives[LightIndex][PrimitiveIndex].Corners[2].WorldPosition - FVector4f(0,0,.1f), true, false));
+								//LightingSystem.DebugOutput.ShadowRays.Add(FDebugStaticLightingRay(EmissivePrimitives[LightIndex][PrimitiveIndex].Corners[2].WorldPosition - FVector4f(0,0,.1f), EmissivePrimitives[LightIndex][PrimitiveIndex].Corners[0].WorldPosition - FVector4f(0,0,.1f), true, true));
 							}
 						}
 					}
@@ -647,13 +647,13 @@ void FStaticLightingMesh::CalculateUniqueLayers(
 				{
 					for (int32 OtherTriangleIndex = 0; OtherTriangleIndex < LayeredGroupTriangles[LayeredGroupIndex].Num(); OtherTriangleIndex++)
 					{
-						const FVector2D& V0 = MeshVertices[GroupedTriangles[GroupIndex][TriangleIndex] * 3 + 0].TextureCoordinates[TextureCoordinateIndex];
-						const FVector2D& V1 = MeshVertices[GroupedTriangles[GroupIndex][TriangleIndex] * 3 + 1].TextureCoordinates[TextureCoordinateIndex];
-						const FVector2D& V2 = MeshVertices[GroupedTriangles[GroupIndex][TriangleIndex] * 3 + 2].TextureCoordinates[TextureCoordinateIndex];
+						const FVector2f& V0 = MeshVertices[GroupedTriangles[GroupIndex][TriangleIndex] * 3 + 0].TextureCoordinates[TextureCoordinateIndex];
+						const FVector2f& V1 = MeshVertices[GroupedTriangles[GroupIndex][TriangleIndex] * 3 + 1].TextureCoordinates[TextureCoordinateIndex];
+						const FVector2f& V2 = MeshVertices[GroupedTriangles[GroupIndex][TriangleIndex] * 3 + 2].TextureCoordinates[TextureCoordinateIndex];
 
-						const FVector2D& OtherV0 = MeshVertices[LayeredGroupTriangles[LayeredGroupIndex][OtherTriangleIndex] * 3 + 0].TextureCoordinates[TextureCoordinateIndex];
-						const FVector2D& OtherV1 = MeshVertices[LayeredGroupTriangles[LayeredGroupIndex][OtherTriangleIndex] * 3 + 1].TextureCoordinates[TextureCoordinateIndex];
-						const FVector2D& OtherV2 = MeshVertices[LayeredGroupTriangles[LayeredGroupIndex][OtherTriangleIndex] * 3 + 2].TextureCoordinates[TextureCoordinateIndex];
+						const FVector2f& OtherV0 = MeshVertices[LayeredGroupTriangles[LayeredGroupIndex][OtherTriangleIndex] * 3 + 0].TextureCoordinates[TextureCoordinateIndex];
+						const FVector2f& OtherV1 = MeshVertices[LayeredGroupTriangles[LayeredGroupIndex][OtherTriangleIndex] * 3 + 1].TextureCoordinates[TextureCoordinateIndex];
+						const FVector2f& OtherV2 = MeshVertices[LayeredGroupTriangles[LayeredGroupIndex][OtherTriangleIndex] * 3 + 2].TextureCoordinates[TextureCoordinateIndex];
 
 						if (AxisAlignedTriangleIntersectTriangle2d(V0, V1, V2, OtherV0, OtherV1, OtherV2))
 						{
@@ -710,7 +710,7 @@ void FStaticLightingMesh::AddLightTexel(
 		{
 			const float XFraction = X / (float)TexSizeX;
 			const float YFraction = Y / (float)TexSizeY;
-			const FLinearColor CurrentEmissive = EvaluateEmissive(FVector2D(XFraction, YFraction), ElementIndex);
+			const FLinearColor CurrentEmissive = EvaluateEmissive(FVector2f(XFraction, YFraction), ElementIndex);
 			if (CurrentEmissive.R > EmissiveThreshold || CurrentEmissive.G > EmissiveThreshold || CurrentEmissive.B > EmissiveThreshold)
 			{
 				Texels.Add(FIntPoint(X, Y));
@@ -729,7 +729,7 @@ void FStaticLightingMesh::AddPrimitiveTexel(
 	const FTexelToCornersMap& TexelToCornersMap, 
 	const FTexelToCorners& ComparisonTexel,
 	int32 ComparisonTexelLightIndex,
-	const FVector4& PrimitiveOrigin,
+	const FVector4f& PrimitiveOrigin,
 	TArray<int32>& PrimitiveIndices, 
 	const TArray<int32>& LightIndices, 
 	int32 X, int32 Y, 
@@ -754,7 +754,7 @@ void FStaticLightingMesh::AddPrimitiveTexel(
 		else if (LightIndex == ComparisonTexelLightIndex)
 		{
 			const FTexelToCorners& CurrentTexelCorners = TexelToCornersMap(X, Y);
-			FVector4 PrimitiveCenter(0,0,0);
+			FVector4f PrimitiveCenter(0,0,0);
 			for (int32 CornerIndex = 0; CornerIndex < NumTexelCorners; CornerIndex++)
 			{
 				PrimitiveCenter += CurrentTexelCorners.Corners[CornerIndex].WorldPosition / (float)NumTexelCorners;
@@ -768,7 +768,7 @@ void FStaticLightingMesh::AddPrimitiveTexel(
 				bool bAnyCornersMatch = false;
 				for (int32 CornerIndex = 0; CornerIndex < NumTexelCorners && !bAnyCornersMatch; CornerIndex++)
 				{
-					const FVector4 CurrentPosition = CurrentTexelCorners.Corners[CornerIndex].WorldPosition;
+					const FVector4f CurrentPosition = CurrentTexelCorners.Corners[CornerIndex].WorldPosition;
 					for (int32 OtherCornerIndex = 0; OtherCornerIndex < NumTexelCorners; OtherCornerIndex++)
 					{
 						if ((CurrentPosition - ComparisonTexel.Corners[OtherCornerIndex].WorldPosition).SizeSquared3() < FMath::Square(Scene.MeshAreaLightSettings.MeshAreaLightSimplifyCornerDistanceThreshold))
diff --git a/Engine/Source/Programs/UnrealLightmass/Private/Lighting/LightingMesh.h b/Engine/Source/Programs/UnrealLightmass/Private/Lighting/LightingMesh.h
--- a/Engine/Source/Programs/UnrealLightmass/Private/Lighting/LightingMesh.h
+++ b/Engine/Source/Programs/UnrealLightmass/Private/Lighting/LightingMesh.h
@@ -35,9 +35,9 @@ struct FStaticLightingVertex: public FStaticLightingVertexData
 	}
 
 	/** Transforms a world space vector into the tangent space of this vertex. */
-	inline FVector4 TransformWorldVectorToTangent(const FVector4& WorldVector) const
+	inline FVector4f TransformWorldVectorToTangent(const FVector4f& WorldVector) const
 	{
-		const FVector4 TangentVector(
+		const FVector4f TangentVector(
 			Dot3(WorldTangentX, WorldVector),
 			Dot3(WorldTangentY, WorldVector),
 			Dot3(WorldTangentZ, WorldVector)
@@ -46,14 +46,14 @@ struct FStaticLightingVertex: public FStaticLightingVertexData
 	}
 
 	/** Transforms a vector in the tangent space of this vertex into world space. */
-	inline FVector4 TransformTangentVectorToWorld(const FVector4& TangentVector) const
+	inline FVector4f TransformTangentVectorToWorld(const FVector4f& TangentVector) const
 	{
 		checkSlow(TangentVector.IsUnit3());
 		// Assuming the transpose of the tangent basis is also the inverse
-		const FVector4 WorldTangentRow0(WorldTangentX.X, WorldTangentY.X, WorldTangentZ.X);
-		const FVector4 WorldTangentRow1(WorldTangentX.Y, WorldTangentY.Y, WorldTangentZ.Y);
-		const FVector4 WorldTangentRow2(WorldTangentX.Z, WorldTangentY.Z, WorldTangentZ.Z);
-		const FVector4 WorldVector(
+		const FVector4f WorldTangentRow0(WorldTangentX.X, WorldTangentY.X, WorldTangentZ.X);
+		const FVector4f WorldTangentRow1(WorldTangentX.Y, WorldTangentY.Y, WorldTangentZ.Y);
+		const FVector4f WorldTangentRow2(WorldTangentX.Z, WorldTangentY.Z, WorldTangentZ.Z);
+		const FVector4f WorldVector(
 			Dot3(WorldTangentRow0, TangentVector),
 			Dot3(WorldTangentRow1, TangentVector),
 			Dot3(WorldTangentRow2, TangentVector)
@@ -68,11 +68,11 @@ struct FStaticLightingVertex: public FStaticLightingVertexData
 		checkSlow(WorldTangentZ.IsUnit3());
 		// Use the vector perpendicular to the normal and the negative Y axis as the TangentX.  
 		// A WorldTangentZ of (0,0,1) will generate WorldTangentX of (1,0,0) and WorldTangentY of (0,1,0) which can be useful for debugging tangent space issues.
-		const FVector4 TangentXCandidate = WorldTangentZ ^ FVector4(0,-1,0);
+		const FVector4f TangentXCandidate = WorldTangentZ ^ FVector4f(0,-1,0);
 		if (TangentXCandidate.SizeSquared3() < KINDA_SMALL_NUMBER)
 		{
 			// The normal was nearly equal to the Y axis, use the X axis instead
-			WorldTangentX = (WorldTangentZ ^ FVector4(1,0,0)).GetUnsafeNormal3();
+			WorldTangentX = (WorldTangentZ ^ FVector4f(1,0,0)).GetUnsafeNormal3();
 		}
 		else
 		{
@@ -148,14 +148,14 @@ struct FStaticLightingVertex: public FStaticLightingVertexData
  */
 struct FFullStaticLightingVertex : public FStaticLightingVertex
 {
-	FVector4 TriangleTangentX;
-	FVector4 TriangleTangentY;
-	FVector4 TriangleNormal;
+	FVector4f TriangleTangentX;
+	FVector4f TriangleTangentY;
+	FVector4f TriangleNormal;
 
 	/** Transforms a world space vector into the tangent space of this triangle. */
-	inline FVector4 TransformWorldVectorToTriangleTangent(const FVector4& WorldVector) const
+	inline FVector4f TransformWorldVectorToTriangleTangent(const FVector4f& WorldVector) const
 	{
-		const FVector4 TangentVector(
+		const FVector4f TangentVector(
 			Dot3(TriangleTangentX, WorldVector),
 			Dot3(TriangleTangentY, WorldVector),
 			Dot3(TriangleNormal, WorldVector)
@@ -164,14 +164,14 @@ struct FFullStaticLightingVertex : public FStaticLightingVertex
 	}
 
 	/** Transforms a vector in the tangent space of this triangle into world space. */
-	inline FVector4 TransformTriangleTangentVectorToWorld(const FVector4& TriangleTangentVector) const
+	inline FVector4f TransformTriangleTangentVectorToWorld(const FVector4f& TriangleTangentVector) const
 	{
 		checkSlow(TriangleTangentVector.IsUnit3());
 		// Assuming the transpose of the tangent basis is also the inverse
-		const FVector4 WorldTangentRow0(TriangleTangentX.X, TriangleTangentY.X, TriangleNormal.X);
-		const FVector4 WorldTangentRow1(TriangleTangentX.Y, TriangleTangentY.Y, TriangleNormal.Y);
-		const FVector4 WorldTangentRow2(TriangleTangentX.Z, TriangleTangentY.Z, TriangleNormal.Z);
-		const FVector4 WorldVector(
+		const FVector4f WorldTangentRow0(TriangleTangentX.X, TriangleTangentY.X, TriangleNormal.X);
+		const FVector4f WorldTangentRow1(TriangleTangentX.Y, TriangleTangentY.Y, TriangleNormal.Y);
+		const FVector4f WorldTangentRow2(TriangleTangentX.Z, TriangleTangentY.Z, TriangleNormal.Z);
+		const FVector4f WorldVector(
 			Dot3(WorldTangentRow0, TriangleTangentVector),
 			Dot3(WorldTangentRow1, TriangleTangentVector),
 			Dot3(WorldTangentRow2, TriangleTangentVector)
@@ -186,11 +186,11 @@ struct FFullStaticLightingVertex : public FStaticLightingVertex
 		checkSlow(TriangleNormal.IsUnit3());
 		// Use the vector perpendicular to the normal and the negative Y axis as the TangentX.  
 		// A TriangleNormal of (0,0,1) will generate TriangleTangentX of (1,0,0) and TriangleTangentY of (0,1,0) which can be useful for debugging tangent space issues.
-		const FVector4 TangentXCandidate = TriangleNormal ^ FVector4(0,-1,0);
+		const FVector4f TangentXCandidate = TriangleNormal ^ FVector4f(0,-1,0);
 		if (TangentXCandidate.SizeSquared3() < KINDA_SMALL_NUMBER)
 		{
 			// The normal was nearly equal to the Y axis, use the X axis instead
-			TriangleTangentX = (TriangleNormal ^ FVector4(1,0,0)).GetUnsafeNormal3();
+			TriangleTangentX = (TriangleNormal ^ FVector4f(1,0,0)).GetUnsafeNormal3();
 		}
 		else
 		{
@@ -202,7 +202,7 @@ struct FFullStaticLightingVertex : public FStaticLightingVertex
 
 	void ApplyVertexModifications(int32 ElementIndex, bool bUseNormalMapsForLighting, const class FStaticLightingMesh* Mesh);
 
-	inline void ComputePathDirections(const FVector4& TriangleTangentPathDirection, FVector4& WorldPathDirection, FVector4& TangentPathDirection) const
+	inline void ComputePathDirections(const FVector4f& TriangleTangentPathDirection, FVector4f& WorldPathDirection, FVector4f& TangentPathDirection) const
 	{
 		checkSlow(TriangleTangentPathDirection.Z >= 0.0f);
 		checkSlow(TriangleTangentPathDirection.IsUnit3());
@@ -440,21 +440,21 @@ public:
 	FLinearColor EvaluateBRDF(
 		const FStaticLightingVertex& Vertex, 
 		int32 ElementIndex,
-		const FVector4& IncomingDirection, 
-		const FVector4& OutgoingDirection) const;
+		const FVector4f& IncomingDirection, 
+		const FVector4f& OutgoingDirection) const;
 
 	/** Generates an outgoing direction sample and evaluates the BRDF for that direction. */
 	FLinearColor SampleBRDF(
 		const FStaticLightingVertex& Vertex, 
 		int32 ElementIndex,
-		const FVector4& IncomingDirection, 
-		FVector4& OutgoingDirection,
+		const FVector4f& IncomingDirection, 
+		FVector4f& OutgoingDirection,
 		float& DirectionPDF,
 		FLMRandomStream& RandomStream
 		) const;
 
 	/** Evaluates the mesh's emissive at the given UVs */
-	inline FLinearColor EvaluateEmissive(const FVector2D& UVs, int32 ElementIndex) const
+	inline FLinearColor EvaluateEmissive(const FVector2f& UVs, int32 ElementIndex) const
 	{
 		checkSlow(IsEmissive(ElementIndex)); 
 		FLinearColor Emissive(FLinearColor::Black);
@@ -472,7 +472,7 @@ public:
 	}
 
 	/** Evaluates the mesh's diffuse at the given UVs */
-	inline FLinearColor EvaluateDiffuse(const FVector2D& UVs, int32 ElementIndex) const
+	inline FLinearColor EvaluateDiffuse(const FVector2f& UVs, int32 ElementIndex) const
 	{
 		checkSlow(!IsTranslucent(ElementIndex));
 		FLinearColor Diffuse(DebugDiffuse);
@@ -497,7 +497,7 @@ public:
 	}
 
 	/** Evaluates the mesh's transmission at the given UVs */
-	inline FLinearColor EvaluateTransmission(const FVector2D& UVs, int32 ElementIndex) const
+	inline FLinearColor EvaluateTransmission(const FVector2f& UVs, int32 ElementIndex) const
 	{
 		checkSlow(IsTranslucent(ElementIndex));
 		FLinearColor Transmission = MaterialElements[ElementIndex].Material->SampleTransmission(UVs);
@@ -508,7 +508,7 @@ public:
 	}
 
 	/** Evaluates the mesh's transmission at the given UVs */
-	inline bool EvaluateMaskedCollision(const FVector2D& UVs, int32 ElementIndex) const
+	inline bool EvaluateMaskedCollision(const FVector2f& UVs, int32 ElementIndex) const
 	{
 		checkSlow(IsMasked(ElementIndex) || IsCastingShadowsAsMasked(ElementIndex));
 		const FMaterialElement& MaterialElement = MaterialElements[ElementIndex];
@@ -518,9 +518,9 @@ public:
 	}
 
 	/** Evaluates the mesh's tangent space normal at the given UVs */
-	inline FVector4 EvaluateNormal(const FVector2D& UVs, int32 ElementIndex) const
+	inline FVector4f EvaluateNormal(const FVector2f& UVs, int32 ElementIndex) const
 	{
-		FVector4 Normal( 0, 0, 1.0f, 0.0 );
+		FVector4f Normal( 0, 0, 1.0f, 0.0 );
 		const FMaterialElement& MaterialElement = MaterialElements[ElementIndex];
 		if( MaterialElement.Material->NormalSize > 0 )
 		{
@@ -565,7 +565,7 @@ private:
 		const FTexelToCornersMap& TexelToCornersMap, 
 		const struct FTexelToCorners& ComparisonTexel,
 		int32 ComparisonTexelLightIndex,
-		const FVector4& PrimitiveOrigin,
+		const FVector4f& PrimitiveOrigin,
 		TArray<int32>& PrimitiveIndices, 
 		const TArray<int32>& LightIndices, 
 		int32 X, int32 Y, 
diff --git a/Engine/Source/Programs/UnrealLightmass/Private/Lighting/LightingSystem.cpp b/Engine/Source/Programs/UnrealLightmass/Private/Lighting/LightingSystem.cpp
--- a/Engine/Source/Programs/UnrealLightmass/Private/Lighting/LightingSystem.cpp
+++ b/Engine/Source/Programs/UnrealLightmass/Private/Lighting/LightingSystem.cpp
@@ -89,7 +89,7 @@ FLightMapData2D* FGatheredLightMapData2D::ConvertToLightmap2D(bool bDebugThisMap
 }
 
 FStaticLightingMappingContext::FStaticLightingMappingContext(const FStaticLightingMesh* InSubjectMesh, FStaticLightingSystem& InSystem, FDebugLightingOutput* InDebugOutput) :
-	FirstBounceCache(InSubjectMesh ? InSubjectMesh->BoundingBox : FBox::BuildAABB(FVector4(0,0,0), FVector4(HALF_WORLD_MAX)), InSystem, 1),
+	FirstBounceCache(InSubjectMesh ? InSubjectMesh->BoundingBox : FBox3f::BuildAABB(FVector4f(0,0,0), FVector4f(HALF_WORLD_MAX)), InSystem, 1),
 	System(InSystem),
 	DebugOutput(InDebugOutput)
 {}
@@ -140,18 +140,18 @@ FStaticLightingSystem::FStaticLightingSystem(const FLightingBuildOptions& InOpti
 ,	NextVolumeSampleTaskIndex(-1)
 ,	NumVolumeSampleTasksOutstanding(0)
 ,	bShouldExportVolumeSampleData(false)
-,	VolumeLightingInterpolationOctree(FVector4(0,0,0), HALF_WORLD_MAX)
+,	VolumeLightingInterpolationOctree(FVector4f(0,0,0), HALF_WORLD_MAX)
 ,	bShouldExportMeshAreaLightData(false)
 ,	bShouldExportVolumeDistanceField(false)
 ,	NumPhotonsEmittedDirect(0)
-,	DirectPhotonMap(FVector4(0,0,0), HALF_WORLD_MAX)
+,	DirectPhotonMap(FVector4f(0,0,0), HALF_WORLD_MAX)
 ,	NumPhotonsEmittedFirstBounce(0)
-,	FirstBouncePhotonMap(FVector4(0,0,0), HALF_WORLD_MAX)
-,	FirstBounceEscapedPhotonMap(FVector4(0,0,0), HALF_WORLD_MAX)
-,	FirstBouncePhotonSegmentMap(FVector4(0,0,0), HALF_WORLD_MAX)
+,	FirstBouncePhotonMap(FVector4f(0,0,0), HALF_WORLD_MAX)
+,	FirstBounceEscapedPhotonMap(FVector4f(0,0,0), HALF_WORLD_MAX)
+,	FirstBouncePhotonSegmentMap(FVector4f(0,0,0), HALF_WORLD_MAX)
 ,	NumPhotonsEmittedSecondBounce(0)
-,	SecondBouncePhotonMap(FVector4(0,0,0), HALF_WORLD_MAX)
-,	IrradiancePhotonMap(FVector4(0,0,0), HALF_WORLD_MAX)
+,	SecondBouncePhotonMap(FVector4f(0,0,0), HALF_WORLD_MAX)
+,	IrradiancePhotonMap(FVector4f(0,0,0), HALF_WORLD_MAX)
 ,	AggregateMesh(NULL)
 ,	VoxelizationSurfaceAggregateMesh(NULL)
 ,	VoxelizationVolumeAggregateMesh(NULL)
@@ -339,7 +339,7 @@ FStaticLightingSystem::FStaticLightingSystem(const FLightingBuildOptions& InOpti
 		for (int32 LightIndex = 0; LightIndex < MeshAreaLights.Num(); LightIndex++)
 		{
 			// Register the newly created mesh area lights with every relevant mesh so they are used for lighting
-			if (MeshAreaLights[LightIndex].AffectsBounds(FBoxSphereBounds(Meshes[MeshIndex]->BoundingBox)))
+			if (MeshAreaLights[LightIndex].AffectsBounds(FBoxSphereBounds3f(Meshes[MeshIndex]->BoundingBox)))
 			{
 				Meshes[MeshIndex]->RelevantLights.Add(&MeshAreaLights[LightIndex]);
 			}
@@ -474,8 +474,8 @@ FStaticLightingSystem::FStaticLightingSystem(const FLightingBuildOptions& InOpti
 
 	GStatistics.NumTotalMappings = Mappings.Num();
 
-	const FBoxSphereBounds SceneBounds = FBoxSphereBounds(AggregateMesh->GetBounds());
-	const FBoxSphereBounds ImportanceBounds = GetImportanceBounds();
+	const FBoxSphereBounds3f SceneBounds = FBoxSphereBounds3f(AggregateMesh->GetBounds());
+	const FBoxSphereBounds3f ImportanceBounds = GetImportanceBounds();
 	// Never trace further than the importance or scene diameter
 	MaxRayDistance = ImportanceBounds.SphereRadius > 0.0f ? ImportanceBounds.SphereRadius * 2.0f : SceneBounds.SphereRadius * 2.0f;
 
@@ -522,7 +522,7 @@ FStaticLightingSystem::FStaticLightingSystem(const FLightingBuildOptions& InOpti
 		Lights.Add(&InScene.RectLights[LightIndex]);
 	}
 
-	const FBoxSphereBounds EffectiveImportanceBounds = ImportanceBounds.SphereRadius > 0.0f ? ImportanceBounds : SceneBounds;
+	const FBoxSphereBounds3f EffectiveImportanceBounds = ImportanceBounds.SphereRadius > 0.0f ? ImportanceBounds : SceneBounds;
 	for (int32 LightIndex = 0; LightIndex < MeshAreaLights.Num(); LightIndex++)
 	{
 		MeshAreaLights[LightIndex].Initialize(Scene.PhotonMappingSettings.IndirectPhotonEmitConeAngle, EffectiveImportanceBounds);
@@ -850,13 +850,13 @@ int32 FStaticLightingSystem::GetNumPhotonImportanceHemisphereSamples() const
 		FMath::TruncToInt(ImportanceTracingSettings.NumHemisphereSamples * PhotonMappingSettings.FinalGatherImportanceSampleFraction) : 0;
 }
 
-FBoxSphereBounds FStaticLightingSystem::GetImportanceBounds(bool bClampToScene) const
+FBoxSphereBounds3f FStaticLightingSystem::GetImportanceBounds(bool bClampToScene) const
 {
-	FBoxSphereBounds ImportanceBounds = Scene.GetImportanceBounds();
+	FBoxSphereBounds3f ImportanceBounds = Scene.GetImportanceBounds();
 	
 	if (bClampToScene)
 	{
-		const FBoxSphereBounds SceneBounds = FBoxSphereBounds(AggregateMesh->GetBounds());
+		const FBoxSphereBounds3f SceneBounds = FBoxSphereBounds3f(AggregateMesh->GetBounds());
 		const float SceneToImportanceOriginSquared = (ImportanceBounds.Origin - SceneBounds.Origin).SizeSquared();
 		if (SceneToImportanceOriginSquared > FMath::Square(SceneBounds.SphereRadius))
 		{
@@ -879,7 +879,7 @@ FBoxSphereBounds FStaticLightingSystem::GetImportanceBounds(bool bClampToScene)
 }
 
 /** Returns true if the specified position is inside any of the importance volumes. */
-bool FStaticLightingSystem::IsPointInImportanceVolume(const FVector4& Position, float Tolerance) const
+bool FStaticLightingSystem::IsPointInImportanceVolume(const FVector4f& Position, float Tolerance) const
 {
 	if (Scene.ImportanceVolumes.Num() > 0)
 	{
@@ -1162,8 +1162,8 @@ void FStaticLightingSystem::DumpStats(float TotalStaticLightingTime) const
 	SolverStats += FString::Printf( TEXT("\n") );
 	SolverStats += FString::Printf( TEXT("Traced %.3f million first hit visibility rays for a total of %.1f thread seconds (%.3f million per thread second)\n"), Stats.NumFirstHitRaysTraced / 1000000.0f, Stats.FirstHitRayTraceThreadTime, Stats.NumFirstHitRaysTraced / 1000000.0f / Stats.FirstHitRayTraceThreadTime);
 	SolverStats += FString::Printf( TEXT("Traced %.3f million boolean visibility rays for a total of %.1f thread seconds (%.3f million per thread second)\n"), Stats.NumBooleanRaysTraced / 1000000.0f, Stats.BooleanRayTraceThreadTime, Stats.NumBooleanRaysTraced / 1000000.0f / Stats.BooleanRayTraceThreadTime);
-	const FBoxSphereBounds SceneBounds = FBoxSphereBounds(AggregateMesh->GetBounds());
-	const FBoxSphereBounds ImportanceBounds = GetImportanceBounds();
+	const FBoxSphereBounds3f SceneBounds = FBoxSphereBounds3f(AggregateMesh->GetBounds());
+	const FBoxSphereBounds3f ImportanceBounds = GetImportanceBounds();
 	SolverStats += FString::Printf( TEXT("Scene radius %.1f, Importance bounds radius %.1f\n"), SceneBounds.SphereRadius, ImportanceBounds.SphereRadius);
 	SolverStats += FString::Printf( TEXT("%u Mappings, %.3f million Texels, %.3f million mapped texels\n"), Stats.NumMappings, Stats.NumTexelsProcessed / 1000000.0f, Stats.NumMappedTexels / 1000000.0f);
 	
@@ -1389,13 +1389,13 @@ void FStaticLightingSystem::CacheSamples()
 	{
 		for (int32 SampleIndex = 0; SampleIndex < NumUniformHemisphereSamples; SampleIndex++)
 		{
-			const FVector4& CurrentSample = GetUniformHemisphereVector(RandomStream, ImportanceTracingSettings.MaxHemisphereRayAngle);
+			const FVector4f& CurrentSample = GetUniformHemisphereVector(RandomStream, ImportanceTracingSettings.MaxHemisphereRayAngle);
 			CachedHemisphereSamples.Add(CurrentSample);
 		}
 	}
 
 	{
-		FVector4 CombinedVector(0);
+		FVector4f CombinedVector(0);
 
 		for (int32 SampleIndex = 0; SampleIndex < CachedHemisphereSamples.Num(); SampleIndex++)
 		{
@@ -1438,7 +1438,7 @@ void FStaticLightingSystem::CacheSamples()
 
 		GenerateStratifiedUniformHemisphereSamples(NumThetaSteps, NumPhiSteps, RandomStream, CachedVolumetricLightmapUniformHemisphereSamples, CachedVolumetricLightmapUniformHemisphereSampleUniforms);
 
-		FVector4 CombinedVector(0);
+		FVector4f CombinedVector(0);
 
 		for (int32 SampleIndex = 0; SampleIndex < CachedVolumetricLightmapUniformHemisphereSamples.Num(); SampleIndex++)
 		{
@@ -1448,7 +1448,7 @@ void FStaticLightingSystem::CacheSamples()
 		CachedVolumetricLightmapMaxUnoccludedLength = (CombinedVector / CachedVolumetricLightmapUniformHemisphereSamples.Num()).Size3();
 	}
 
-	CachedVolumetricLightmapVertexOffsets.Add(FVector(0, 0, 0));
+	CachedVolumetricLightmapVertexOffsets.Add(FVector3f(0, 0, 0));
 }
 
 bool FStaticLightingThreadRunnable::CheckHealth(bool bReportError) const
@@ -2042,7 +2042,7 @@ class FStoredLightingSample
 {
 public:
 	FLinearColor IncomingRadiance;
-	FVector4 WorldSpaceDirection;
+	FVector4f WorldSpaceDirection;
 };
 
 class FSampleCollector
@@ -2052,7 +2052,7 @@ public:
 	inline void SetOcclusion(float InOcclusion)
 	{}
 
-	inline void AddIncomingRadiance(const FLinearColor& IncomingRadiance, float Weight, const FVector4& TangentSpaceDirection, const FVector4& WorldSpaceDirection)
+	inline void AddIncomingRadiance(const FLinearColor& IncomingRadiance, float Weight, const FVector4f& TangentSpaceDirection, const FVector4f& WorldSpaceDirection)
 	{
 		if (FLinearColorUtils::LinearRGBToXYZ(IncomingRadiance * Weight).G > DELTA)
 		{
@@ -2099,13 +2099,13 @@ void FStaticLightingSystem::CalculateStaticShadowDepthMap(FGuid LightGuid)
 
 	if (DirectionalLight)
 	{
-		FVector4 XAxis, YAxis;
+		FVector4f XAxis, YAxis;
 		DirectionalLight->Direction.FindBestAxisVectors3(XAxis, YAxis);
 		// Create a coordinate system for the dominant directional light, with the z axis corresponding to the light's direction
-		ShadowDepthMap->WorldToLight = FBasisVectorMatrix(XAxis, YAxis, DirectionalLight->Direction, FVector4(0,0,0));
+		ShadowDepthMap->WorldToLight = FBasisVectorMatrix44f(XAxis, YAxis, DirectionalLight->Direction, FVector4f(0,0,0));
 
-		FBoxSphereBounds ImportanceVolume = GetImportanceBounds().SphereRadius > 0.0f ? GetImportanceBounds() : FBoxSphereBounds(AggregateMesh->GetBounds());
-		const FBox LightSpaceImportanceBounds = ImportanceVolume.GetBox().TransformBy(ShadowDepthMap->WorldToLight);
+		FBoxSphereBounds3f ImportanceVolume = GetImportanceBounds().SphereRadius > 0.0f ? GetImportanceBounds() : FBoxSphereBounds3f(AggregateMesh->GetBounds());
+		const FBox3f LightSpaceImportanceBounds = ImportanceVolume.GetBox().TransformBy(ShadowDepthMap->WorldToLight);
 
 		ShadowDepthMap->ShadowMapSizeX = FMath::TruncToInt(FMath::Max(LightSpaceImportanceBounds.GetExtent().X * 2.0f * ClampedResolutionScale / ShadowSettings.StaticShadowDepthMapTransitionSampleDistanceX, 4.0f));
 		ShadowDepthMap->ShadowMapSizeX = ShadowDepthMap->ShadowMapSizeX == appTruncErrorCode ? INT_MAX : ShadowDepthMap->ShadowMapSizeX;
@@ -2126,7 +2126,7 @@ void FStaticLightingSystem::CalculateStaticShadowDepthMap(FGuid LightGuid)
 
 		{
 			const float InvDistanceRange = 1.0f / (LightSpaceImportanceBounds.Max.Z - LightSpaceImportanceBounds.Min.Z);
-			const FMatrix LightToWorld = ShadowDepthMap->WorldToLight.InverseFast();
+			const FMatrix44f LightToWorld = ShadowDepthMap->WorldToLight.InverseFast();
 
 			for (int32 Y = 0; Y < ShadowDepthMap->ShadowMapSizeY; Y++)
 			{
@@ -2141,14 +2141,14 @@ void FStaticLightingSystem::CalculateStaticShadowDepthMap(FGuid LightGuid)
 						{
 							const float XFraction = (X + SubSampleX / (float)ShadowSettings.StaticShadowDepthMapSuperSampleFactor) / (float)(ShadowDepthMap->ShadowMapSizeX - 1);
 							// Construct a ray in light space along the direction of the light, starting at the minimum light space Z going to the maximum.
-							const FVector4 LightSpaceStartPosition(
+							const FVector4f LightSpaceStartPosition(
 								LightSpaceImportanceBounds.Min.X + XFraction * (LightSpaceImportanceBounds.Max.X - LightSpaceImportanceBounds.Min.X),
 								LightSpaceImportanceBounds.Min.Y + YFraction * (LightSpaceImportanceBounds.Max.Y - LightSpaceImportanceBounds.Min.Y),
 								LightSpaceImportanceBounds.Min.Z);
-							const FVector4 LightSpaceEndPosition(LightSpaceStartPosition.X, LightSpaceStartPosition.Y, LightSpaceImportanceBounds.Max.Z);
+							const FVector4f LightSpaceEndPosition(LightSpaceStartPosition.X, LightSpaceStartPosition.Y, LightSpaceImportanceBounds.Max.Z);
 							// Transform the ray into world space in order to trace against the world space aggregate mesh
-							const FVector4 WorldSpaceStartPosition = LightToWorld.TransformPosition(LightSpaceStartPosition);
-							const FVector4 WorldSpaceEndPosition = LightToWorld.TransformPosition(LightSpaceEndPosition);
+							const FVector4f WorldSpaceStartPosition = LightToWorld.TransformPosition(LightSpaceStartPosition);
+							const FVector4f WorldSpaceEndPosition = LightToWorld.TransformPosition(LightSpaceEndPosition);
 							const FLightRay LightRay(
 								WorldSpaceStartPosition,
 								WorldSpaceEndPosition,
@@ -2181,24 +2181,24 @@ void FStaticLightingSystem::CalculateStaticShadowDepthMap(FGuid LightGuid)
 		}
 
 		ShadowDepthMap->WorldToLight *= FTranslationMatrix(-LightSpaceImportanceBounds.Min)
-			* FScaleMatrix(FVector(1.0f) / (LightSpaceImportanceBounds.Max - LightSpaceImportanceBounds.Min));
+			* FScaleMatrix(FVector3f(1.0f) / (LightSpaceImportanceBounds.Max - LightSpaceImportanceBounds.Min));
 
 		FScopeLock Lock(&CompletedStaticShadowDepthMapsSync);
 		CompletedStaticShadowDepthMaps.Add(DirectionalLight, ShadowDepthMap);
 	}
 	else if (SpotLight)
 	{
-		FVector4 XAxis, YAxis;
+		FVector4f XAxis, YAxis;
 		SpotLight->Direction.FindBestAxisVectors3(XAxis, YAxis);
 		// Create a coordinate system for the spot light, with the z axis corresponding to the light's direction, and translated to the light's origin
-		ShadowDepthMap->WorldToLight = FTranslationMatrix(-SpotLight->Position) 
-			* FBasisVectorMatrix(XAxis, YAxis, SpotLight->Direction, FVector4(0,0,0));
+		ShadowDepthMap->WorldToLight = FTranslationMatrix44f(-SpotLight->Position) 
+			* FBasisVectorMatrix44f(XAxis, YAxis, SpotLight->Direction, FVector4f(0,0,0));
 
 		// Distance from the light's direction axis to the edge of the cone at the radius of the light
 		const float HalfCrossSectionLength = SpotLight->Radius * FMath::Tan(SpotLight->OuterConeAngle * (float)PI / 180.0f);
 
-		const FVector4 LightSpaceImportanceBoundMin = FVector4(-HalfCrossSectionLength, -HalfCrossSectionLength, 0);
-		const FVector4 LightSpaceImportanceBoundMax = FVector4(HalfCrossSectionLength, HalfCrossSectionLength, SpotLight->Radius);
+		const FVector4f LightSpaceImportanceBoundMin = FVector4f(-HalfCrossSectionLength, -HalfCrossSectionLength, 0);
+		const FVector4f LightSpaceImportanceBoundMax = FVector4f(HalfCrossSectionLength, HalfCrossSectionLength, SpotLight->Radius);
 
 		ShadowDepthMap->ShadowMapSizeX = FMath::TruncToInt(FMath::Max(HalfCrossSectionLength * ClampedResolutionScale / ShadowSettings.StaticShadowDepthMapTransitionSampleDistanceX, 4.0f));
 		ShadowDepthMap->ShadowMapSizeX = ShadowDepthMap->ShadowMapSizeX == appTruncErrorCode ? INT_MAX : ShadowDepthMap->ShadowMapSizeX;
@@ -2217,8 +2217,8 @@ void FStaticLightingSystem::CalculateStaticShadowDepthMap(FGuid LightGuid)
 
 		// Calculate the maximum possible distance for quantization
 		const float MaxPossibleDistance = LightSpaceImportanceBoundMax.Z - LightSpaceImportanceBoundMin.Z;
-		const FMatrix LightToWorld = ShadowDepthMap->WorldToLight.InverseFast();
-		const FBoxSphereBounds ImportanceVolume = GetImportanceBounds().SphereRadius > 0.0f ? GetImportanceBounds() : FBoxSphereBounds(AggregateMesh->GetBounds());
+		const FMatrix44f LightToWorld = ShadowDepthMap->WorldToLight.InverseFast();
+		const FBoxSphereBounds3f ImportanceVolume = GetImportanceBounds().SphereRadius > 0.0f ? GetImportanceBounds() : FBoxSphereBounds3f(AggregateMesh->GetBounds());
 
 		for (int32 Y = 0; Y < ShadowDepthMap->ShadowMapSizeY; Y++)
 		{
@@ -2233,14 +2233,14 @@ void FStaticLightingSystem::CalculateStaticShadowDepthMap(FGuid LightGuid)
 					{
 						const float XFraction = (X + SubSampleX / (float)ShadowSettings.StaticShadowDepthMapSuperSampleFactor) / (float)(ShadowDepthMap->ShadowMapSizeX - 1);
 						// Construct a ray in light space along the direction of the light, starting at the light and going to the maximum light space Z.
-						const FVector4 LightSpaceStartPosition(0,0,0);
-						const FVector4 LightSpaceEndPosition(
+						const FVector4f LightSpaceStartPosition(0,0,0);
+						const FVector4f LightSpaceEndPosition(
 							LightSpaceImportanceBoundMin.X + XFraction * (LightSpaceImportanceBoundMax.X - LightSpaceImportanceBoundMin.X),
 							LightSpaceImportanceBoundMin.Y + YFraction * (LightSpaceImportanceBoundMax.Y - LightSpaceImportanceBoundMin.Y),
 							LightSpaceImportanceBoundMax.Z);
 						// Transform the ray into world space in order to trace against the world space aggregate mesh
-						const FVector4 WorldSpaceStartPosition = LightToWorld.TransformPosition(LightSpaceStartPosition);
-						const FVector4 WorldSpaceEndPosition = LightToWorld.TransformPosition(LightSpaceEndPosition);
+						const FVector4f WorldSpaceStartPosition = LightToWorld.TransformPosition(LightSpaceStartPosition);
+						const FVector4f WorldSpaceEndPosition = LightToWorld.TransformPosition(LightSpaceEndPosition);
 						const FLightRay LightRay(
 							WorldSpaceStartPosition,
 							WorldSpaceEndPosition,
@@ -2256,7 +2256,7 @@ void FStaticLightingSystem::CalculateStaticShadowDepthMap(FGuid LightGuid)
 
 						if (Intersection.bIntersects)
 						{
-							const FVector4 LightSpaceIntersectPosition = ShadowDepthMap->WorldToLight.TransformPosition(Intersection.IntersectionVertex.WorldPosition);
+							const FVector4f LightSpaceIntersectPosition = ShadowDepthMap->WorldToLight.TransformPosition(Intersection.IntersectionVertex.WorldPosition);
 							// Use the maximum distance of all super samples for each cell, to get a conservative shadow map
 							MaxSampleDistance = FMath::Max(MaxSampleDistance, LightSpaceIntersectPosition.Z);
 						}
@@ -2276,11 +2276,11 @@ void FStaticLightingSystem::CalculateStaticShadowDepthMap(FGuid LightGuid)
 			// Perspective projection sized to the spotlight cone
 			FPerspectiveMatrix(SpotLight->OuterConeAngle * (float)PI / 180.0f, 1, 1, 0, SpotLight->Radius)
 			// Convert from NDC to texture space, normalize Z
-			* FMatrix(
-				FPlane(.5f,								0,							0,									0),
-				FPlane(0,								.5f,						0,									0),
-				FPlane(0,								0,							1.0f / LightSpaceImportanceBoundMax.Z,		0),
-				FPlane(.5f,								.5f,						0,									1));
+			* FMatrix44f(
+				FPlane4f(.5f,								0,							0,									0),
+				FPlane4f(0,								.5f,						0,									0),
+				FPlane4f(0,								0,							1.0f / LightSpaceImportanceBoundMax.Z,		0),
+				FPlane4f(.5f,								.5f,						0,									1));
 
 		FScopeLock Lock(&CompletedStaticShadowDepthMapsSync);
 		CompletedStaticShadowDepthMaps.Add(SpotLight, ShadowDepthMap);
@@ -2303,7 +2303,7 @@ void FStaticLightingSystem::CalculateStaticShadowDepthMap(FGuid LightGuid)
 		ShadowDepthMap->ShadowMap.Empty(ShadowDepthMap->ShadowMapSizeX * ShadowDepthMap->ShadowMapSizeY);
 		ShadowDepthMap->ShadowMap.AddZeroed(ShadowDepthMap->ShadowMapSizeX * ShadowDepthMap->ShadowMapSizeY);
 
-		ShadowDepthMap->WorldToLight = FMatrix::Identity;
+		ShadowDepthMap->WorldToLight = FMatrix44f::Identity;
 
 		for (int32 Y = 0; Y < ShadowDepthMap->ShadowMapSizeY; Y++)
 		{
@@ -2321,10 +2321,10 @@ void FStaticLightingSystem::CalculateStaticShadowDepthMap(FGuid LightGuid)
 					{
 						const float XFraction = (X + SubSampleX / (float)ShadowSettings.StaticShadowDepthMapSuperSampleFactor) / (float)(ShadowDepthMap->ShadowMapSizeX - 1);
 						const float Theta = XFraction * 2 * PI;
-						const FVector Direction(FMath::Cos(Theta) * SinPhi, FMath::Sin(Theta) * SinPhi, FMath::Cos(Phi));
+						const FVector3f Direction(FMath::Cos(Theta) * SinPhi, FMath::Sin(Theta) * SinPhi, FMath::Cos(Phi));
 
-						const FVector4 WorldSpaceStartPosition = PointLight->Position;
-						const FVector4 WorldSpaceEndPosition = PointLight->Position + Direction * PointLight->Radius;
+						const FVector4f WorldSpaceStartPosition = PointLight->Position;
+						const FVector4f WorldSpaceEndPosition = PointLight->Position + Direction * PointLight->Radius;
 						const FLightRay LightRay(
 							WorldSpaceStartPosition,
 							WorldSpaceEndPosition,
@@ -2374,7 +2374,7 @@ void FStaticLightingSystem::CalculateStaticShadowDepthMap(FGuid LightGuid)
  */
 bool FStaticLightingSystem::CalculatePointShadowing(
 	const FStaticLightingMapping* Mapping,
-	const FVector4& WorldSurfacePoint,
+	const FVector4f& WorldSurfacePoint,
 	const FLight* Light,
 	FStaticLightingMappingContext& MappingContext,
 	bool bDebugThisSample
@@ -2387,7 +2387,7 @@ bool FStaticLightingSystem::CalculatePointShadowing(
 	else
 	{
 		// Treat points which the light doesn't affect as shadowed to avoid the costly ray check.
-		if(!Light->AffectsBounds(FBoxSphereBounds(WorldSurfacePoint,FVector4(0,0,0,0),0)))
+		if(!Light->AffectsBounds(FBoxSphereBounds3f(WorldSurfacePoint,FVector4f(0,0,0,0),0)))
 		{
 			return true;
 		}
@@ -2398,8 +2398,8 @@ bool FStaticLightingSystem::CalculatePointShadowing(
 		{
 			// TODO find best point on light to shadow from
 			// Construct a line segment between the light and the surface point.
-			const FVector4 LightPosition = FVector4(Light->Position.X, Light->Position.Y, Light->Position.Z, 0);
-			const FVector4 LightVector = LightPosition - WorldSurfacePoint * Light->Position.W;
+			const FVector4f LightPosition = FVector4f(Light->Position.X, Light->Position.Y, Light->Position.Z, 0);
+			const FVector4f LightVector = LightPosition - WorldSurfacePoint * Light->Position.W;
 			const FLightRay LightRay(
 				WorldSurfacePoint + LightVector.GetSafeNormal() * SceneConstants.VisibilityRayOffsetDistance,
 				WorldSurfacePoint + LightVector,
@@ -2430,7 +2430,7 @@ bool FStaticLightingSystem::CalculatePointShadowing(
 }
 
 /** Calculates area shadowing from a light for the given vertex. */
-FVector2D FStaticLightingSystem::CalculatePointAreaShadowing(
+FVector2f FStaticLightingSystem::CalculatePointAreaShadowing(
 	const FStaticLightingMapping* Mapping,
 	const FStaticLightingVertex& Vertex,
 	int32 ElementIndex,
@@ -2452,9 +2452,9 @@ FVector2D FStaticLightingSystem::CalculatePointAreaShadowing(
 
 	UnnormalizedTransmission = FLinearColor::Black;
 	// Treat points which the light doesn't affect as shadowed to avoid the costly ray check.
-	if( !Light->AffectsBounds(FBoxSphereBounds(Vertex.WorldPosition,FVector4(0,0,0),0)))
+	if( !Light->AffectsBounds(FBoxSphereBounds3f(Vertex.WorldPosition,FVector4f(0,0,0),0)))
 	{
-		return FVector2D::ZeroVector;
+		return FVector2f::ZeroVector;
 	}
 
 	// Check for visibility between the point and the light
@@ -2462,7 +2462,7 @@ FVector2D FStaticLightingSystem::CalculatePointAreaShadowing(
 	{
 		MappingContext.Stats.NumDirectLightingShadowRays += LightPositionSamples.Num();
 		const bool bIsTwoSided = Mapping->Mesh->IsTwoSided(ElementIndex);
-		FVector2D ShadowedValue = FVector2D::ZeroVector;
+		FVector2f ShadowedValue = FVector2f::ZeroVector;
 
 		// Integrate over the surface of the light using monte carlo integration
 		// Note that we are making the approximation that the BRDF and the Light's emission are equal in all of these directions and therefore are not in the integrand
@@ -2473,8 +2473,8 @@ FVector2D FStaticLightingSystem::CalculatePointAreaShadowing(
 			Light->ValidateSurfaceSample(Vertex.WorldPosition, CurrentSample);
 
 			// Construct a line segment between the light and the surface point.
-			const FVector4 LightVector = CurrentSample.Position - Vertex.WorldPosition;
-			FVector4 SampleOffset(0,0,0);
+			const FVector4f LightVector = CurrentSample.Position - Vertex.WorldPosition;
+			FVector4f SampleOffset(0,0,0);
 			if (GeneralSettings.bAccountForTexelSize)
 			{
 				/*
@@ -2492,7 +2492,7 @@ FVector2D FStaticLightingSystem::CalculatePointAreaShadowing(
 				continue;
 			}
 
-			FVector4 NormalForOffset = Vertex.WorldTangentZ;
+			FVector4f NormalForOffset = Vertex.WorldTangentZ;
 			// Flip the normal used for offsetting the start of the ray for two sided materials if a flipped normal would be closer to the light.
 			// This prevents incorrect shadowing where using the frontface normal would cause the ray to start inside a nearby object.
 			if (bIsTwoSided && Dot3(-NormalForOffset, LightVector) > Dot3(NormalForOffset, LightVector))
@@ -2539,7 +2539,7 @@ FVector2D FStaticLightingSystem::CalculatePointAreaShadowing(
 		return ShadowedValue;
 	}
 	UnnormalizedTransmission = FLinearColor::White;
-	return FVector2D( 1.0f, 1.0f );
+	return FVector2f( 1.0f, 1.0f );
 }
 
 /** Calculates the lighting contribution of a light to a mapping vertex. */
@@ -2556,11 +2556,11 @@ FGatheredLightSample FStaticLightingSystem::CalculatePointLighting(
 	if (Light->GetSkyLight() == NULL)
 	{
 	    // Calculate the direction from the vertex to the light.
-		const FVector4 WorldLightVector = Light->GetDirectLightingDirection(Vertex.WorldPosition, Vertex.WorldTangentZ);
+		const FVector4f WorldLightVector = Light->GetDirectLightingDirection(Vertex.WorldPosition, Vertex.WorldTangentZ);
     
 	    // Transform the light vector to tangent space.
-	    const FVector4 TangentLightVector = 
-		    FVector4(
+	    const FVector4f TangentLightVector = 
+		    FVector4f(
 			    Dot3(WorldLightVector, Vertex.WorldTangentX),
 			    Dot3(WorldLightVector, Vertex.WorldTangentY),
 			    Dot3(WorldLightVector, Vertex.WorldTangentZ),
@@ -2575,8 +2575,8 @@ FGatheredLightSample FStaticLightingSystem::CalculatePointLighting(
 
 		if (Mapping->Mesh->UsesTwoSidedLighting(ElementIndex))
 		{
-			const FVector4 BackFaceTangentLightVector = 
-				FVector4(
+			const FVector4f BackFaceTangentLightVector = 
+				FVector4f(
 					Dot3(WorldLightVector, -Vertex.WorldTangentX),
 					Dot3(WorldLightVector, -Vertex.WorldTangentY),
 					Dot3(WorldLightVector, -Vertex.WorldTangentZ),
@@ -2619,7 +2619,7 @@ FGatheredLightSample FStaticLightingSystem::GetVisualizedMaterialAttribute(const
 		}
 		else if (MaterialSettings.ViewMaterialAttribute == VMA_Normal)
 		{
-			const FVector4 Normal = Intersection.Mesh->EvaluateNormal(Intersection.IntersectionVertex.TextureCoordinates[0], Intersection.ElementIndex);
+			const FVector4f Normal = Intersection.Mesh->EvaluateNormal(Intersection.IntersectionVertex.TextureCoordinates[0], Intersection.ElementIndex);
 
 			FLinearColor NormalColor;
 			NormalColor.R = Normal.X * 0.5f + 0.5f;
@@ -2651,14 +2651,14 @@ FGatheredLightSample FStaticLightingSystem::GetVisualizedMaterialAttribute(const
  * @param Light - The light to classify.
  * @return true if the light is behind the triangle.
  */
-bool IsLightBehindSurface(const FVector4& TrianglePoint, const FVector4& TriangleNormal, const FLight* Light)
+bool IsLightBehindSurface(const FVector4f& TrianglePoint, const FVector4f& TriangleNormal, const FLight* Light)
 {
 	const bool bIsSkyLight = Light->GetSkyLight() != NULL;
 	if (!bIsSkyLight)
 	{
 		// Calculate the direction from the triangle to the light.
-		const FVector4 LightPosition = FVector4(Light->Position.X, Light->Position.Y, Light->Position.Z, 0);
-		const FVector4 WorldLightVector = LightPosition - TrianglePoint * Light->Position.W;
+		const FVector4f LightPosition = FVector4f(Light->Position.X, Light->Position.Y, Light->Position.Z, 0);
+		const FVector4f WorldLightVector = LightPosition - TrianglePoint * Light->Position.W;
 
 		// Check if the light is in front of the triangle.
 		const float Dot = Dot3(WorldLightVector, TriangleNormal);
@@ -2679,7 +2679,7 @@ bool IsLightBehindSurface(const FVector4& TrianglePoint, const FVector4& Triangl
  * @param Lights - The lights to cull.
  * @return A map from Lights index to a boolean which is true if the light is in front of the triangle.
  */
-TBitArray<> CullBackfacingLights(bool bTwoSidedMaterial,const FVector4& TrianglePoint,const FVector4& TriangleNormal,const TArray<FLight*>& Lights)
+TBitArray<> CullBackfacingLights(bool bTwoSidedMaterial,const FVector4f& TrianglePoint,const FVector4f& TriangleNormal,const TArray<FLight*>& Lights)
 {
 	if(!bTwoSidedMaterial)
 	{
