Fix CPU Lightmass by coverting LWC types to single precision

diff --git a/Engine/Source/Programs/UnrealLightmass/Private/Lighting/LightingSystem.h b/Engine/Source/Programs/UnrealLightmass/Private/Lighting/LightingSystem.h
--- a/Engine/Source/Programs/UnrealLightmass/Private/Lighting/LightingSystem.h
+++ b/Engine/Source/Programs/UnrealLightmass/Private/Lighting/LightingSystem.h
@@ -95,7 +95,7 @@ struct FFinalGatherInfo
 
 struct FTexelCorner
 {
-	FVector4 WorldPosition;
+	FVector4f WorldPosition;
 };
 
 /** Information about a texel's corners */
@@ -104,9 +104,9 @@ struct FTexelToCorners
 	/** The position of each corner */
 	FTexelCorner Corners[NumTexelCorners];
 	/** The tangent basis of the last valid corner to be rasterized */
-	FVector4 WorldTangentX;
-	FVector4 WorldTangentY;
-	FVector4 WorldTangentZ;
+	FVector4f WorldTangentX;
+	FVector4f WorldTangentY;
+	FVector4f WorldTangentZ;
 	/** Whether each corner lies on the mesh */
 	bool bValid[NumTexelCorners];
 };
@@ -159,23 +159,23 @@ class FPhoton
 private:
 
 	/** Position that the photon was deposited at in XYZ, and Id in W for debugging. */
-	FVector4 PositionAndId;
+	FVector4f PositionAndId;
 
 	/** Direction the photon came from in XYZ, and distance that the photon traveled along its last path before being deposited in W. */
-	FVector4 IncidentDirectionAndDistance;
+	FVector4f IncidentDirectionAndDistance;
 
 	/** Normal of the surface the photon was deposited on in XYZ, and fraction of the originating light's power that this photon represents in W. */
-	FVector4 SurfaceNormalAndPower;
+	FVector4f SurfaceNormalAndPower;
 
 public:
 
-	FPhoton(int32 InId, const FVector4& InPosition, float InDistance, const FVector4& InIncidentDirection, const FVector4& InSurfaceNormal, const FLinearColor& InPower)
+	FPhoton(int32 InId, const FVector4f& InPosition, float InDistance, const FVector4f& InIncidentDirection, const FVector4f& InSurfaceNormal, const FLinearColor& InPower)
 	{
-		PositionAndId = FVector4(InPosition, *(float*)&InId);
-		IncidentDirectionAndDistance = FVector4(InIncidentDirection, InDistance);
+		PositionAndId = FVector4f(InPosition, *(float*)&InId);
+		IncidentDirectionAndDistance = FVector4f(InIncidentDirection, InDistance);
 		checkSlow(FLinearColorUtils::AreFloatsValid(InPower));
 		const FColor PowerRGBE = InPower.ToRGBE();
-		SurfaceNormalAndPower = FVector4(InSurfaceNormal, *(float*)&PowerRGBE);
+		SurfaceNormalAndPower = FVector4f(InSurfaceNormal, *(float*)&PowerRGBE);
 	}
 
 	FORCEINLINE int32 GetId() const
@@ -183,14 +183,14 @@ public:
 		return *(int32*)&PositionAndId.W;
 	}
 
-	FORCEINLINE FVector4 GetPosition() const
+	FORCEINLINE FVector4f GetPosition() const
 	{
-		return FVector4(PositionAndId, 0.0f);
+		return FVector4f(PositionAndId, 0.0f);
 	}
 
-	FORCEINLINE FVector4 GetIncidentDirection() const
+	FORCEINLINE FVector4f GetIncidentDirection() const
 	{
-		return FVector4(IncidentDirectionAndDistance, 0.0f);
+		return FVector4f(IncidentDirectionAndDistance, 0.0f);
 	}
 
 	FORCEINLINE float GetDistance() const
@@ -198,9 +198,9 @@ public:
 		return IncidentDirectionAndDistance.W;
 	}
 
-	FORCEINLINE FVector4 GetSurfaceNormal() const
+	FORCEINLINE FVector4f GetSurfaceNormal() const
 	{
-		return FVector4(SurfaceNormalAndPower, 0.0f);
+		return FVector4f(SurfaceNormalAndPower, 0.0f);
 	}
 
 	FORCEINLINE FLinearColor GetPower() const
@@ -239,7 +239,7 @@ struct FPhotonMapOctreeSemantics
 
 	static FBoxCenterAndExtent GetBoundingBox(const FPhotonElement& PhotonElement)
 	{
-		return FBoxCenterAndExtent(PhotonElement.Photon.GetPosition(), FVector4(0,0,0));
+		return FBoxCenterAndExtent(PhotonElement.Photon.GetPosition(), FVector4f(0,0,0));
 	}
 };
 
@@ -247,18 +247,18 @@ struct FPhotonSegmentElement
 {
 	const FPhoton* Photon;
 
-	FVector SegmentCenter;
-	FVector SegmentExtent;
+	FVector3f SegmentCenter;
+	FVector3f SegmentExtent;
 
 	/** Initialization constructor. */
 	FPhotonSegmentElement(const FPhoton* InPhoton, float InStartOffset, float InSegmentLength) :
 		Photon(InPhoton)
 	{
-		const FVector PhotonDirection = Photon->GetIncidentDirection() * Photon->GetDistance();
-		const FVector SegmentStart = Photon->GetPosition() + PhotonDirection * InStartOffset;
-		const FVector SegmentEnd = SegmentStart + PhotonDirection * InSegmentLength;
+		const FVector3f PhotonDirection = Photon->GetIncidentDirection() * Photon->GetDistance();
+		const FVector3f SegmentStart = Photon->GetPosition() + PhotonDirection * InStartOffset;
+		const FVector3f SegmentEnd = SegmentStart + PhotonDirection * InSegmentLength;
 
-		FBox SegmentBounds(ForceInit);
+		FBox3f SegmentBounds(ForceInit);
 		SegmentBounds += SegmentStart;
 		SegmentBounds += SegmentEnd;
 
@@ -267,11 +267,11 @@ struct FPhotonSegmentElement
 		SegmentExtent = SegmentBounds.GetExtent();
 	}
 
-	inline float ComputeSquaredDistanceToPoint(FVector InPoint) const
+	inline float ComputeSquaredDistanceToPoint(FVector3f InPoint) const
 	{
-		float Projection = FVector::DotProduct(InPoint - Photon->GetPosition(), Photon->GetIncidentDirection());
+		float Projection = FVector3f::DotProduct(InPoint - Photon->GetPosition(), Photon->GetIncidentDirection());
 		Projection = FMath::Clamp(Projection, 0.0f, Photon->GetDistance());
-		FVector ProjectedPosition = Photon->GetPosition() + Photon->GetIncidentDirection() * Projection;
+		FVector3f ProjectedPosition = Photon->GetPosition() + Photon->GetIncidentDirection() * Projection;
 		return (InPoint - ProjectedPosition).SizeSquared();
 	}
 };
@@ -299,10 +299,10 @@ class FIrradiancePhoton : public FIrradiancePhotonData
 {
 public:
 
-	FIrradiancePhoton(const FVector4& InPosition, const FVector4& InSurfaceNormal, bool bInHasContributionFromDirectPhotons)
+	FIrradiancePhoton(const FVector4f& InPosition, const FVector4f& InSurfaceNormal, bool bInHasContributionFromDirectPhotons)
 	{
-		PositionAndDirectContribution = FVector4(InPosition, bInHasContributionFromDirectPhotons);
-		SurfaceNormalAndIrradiance = FVector4(InSurfaceNormal, 0.0f);
+		PositionAndDirectContribution = FVector4f(InPosition, bInHasContributionFromDirectPhotons);
+		SurfaceNormalAndIrradiance = FVector4f(InSurfaceNormal, 0.0f);
 	}
 
 	FORCEINLINE bool HasDirectContribution() const
@@ -340,14 +340,14 @@ public:
 		return OutIrradiance;
 	}
 
-	FORCEINLINE FVector4 GetPosition() const
+	FORCEINLINE FVector4f GetPosition() const
 	{
-		return FVector4(PositionAndDirectContribution, 0.0f);
+		return FVector4f(PositionAndDirectContribution, 0.0f);
 	}
 
-	FORCEINLINE FVector4 GetSurfaceNormal() const
+	FORCEINLINE FVector4f GetSurfaceNormal() const
 	{
-		return FVector4(SurfaceNormalAndIrradiance, 0.0f);
+		return FVector4f(SurfaceNormalAndIrradiance, 0.0f);
 	}
 };
 
@@ -383,7 +383,7 @@ struct FIrradiancePhotonMapOctreeSemantics
 
 	static FBoxCenterAndExtent GetBoundingBox(const FIrradiancePhotonElement& PhotonElement)
 	{
-		return FBoxCenterAndExtent(PhotonElement.GetPhoton().GetPosition(), FVector4(0,0,0));
+		return FBoxCenterAndExtent(PhotonElement.GetPhoton().GetPosition(), FVector4f(0,0,0));
 	}
 };
 
@@ -391,7 +391,7 @@ struct FIrradiancePhotonMapOctreeSemantics
 class FVolumeLightingSample : public FVolumeLightingSampleData
 {
 public:
-	FVolumeLightingSample(const FVector4& InPositionAndRadius)
+	FVolumeLightingSample(const FVector4f& InPositionAndRadius)
 	{
 		for (int32 CoefficientIndex = 0; CoefficientIndex < LM_NUM_SH_COEFFICIENTS; CoefficientIndex++)
 		{
@@ -404,9 +404,9 @@ public:
 
 		PositionAndRadius = InPositionAndRadius;
 	}
-	inline FVector4 GetPosition() const
+	inline FVector4f GetPosition() const
 	{
-		return FVector4(PositionAndRadius, 0.0f);
+		return FVector4f(PositionAndRadius, 0.0f);
 	}
 	inline float GetRadius() const
 	{
@@ -446,14 +446,14 @@ struct FVolumeLightingInterpolationOctreeSemantics
 	static FBoxCenterAndExtent GetBoundingBox(const FVolumeSampleInterpolationElement& Element)
 	{
 		const FVolumeLightingSample& Sample = Element.VolumeSamples[Element.SampleIndex];
-		return FBoxCenterAndExtent(FVector4(Sample.PositionAndRadius, 0.0f), FVector4(Sample.PositionAndRadius.W, Sample.PositionAndRadius.W, Sample.PositionAndRadius.W));
+		return FBoxCenterAndExtent(FVector4f(Sample.PositionAndRadius, 0.0f), FVector4f(Sample.PositionAndRadius.W, Sample.PositionAndRadius.W, Sample.PositionAndRadius.W));
 	}
 };
 
 class FPrecomputedVisibilityCell
 {
 public:
-	FBox Bounds;
+	FBox3f Bounds;
 	TArray<uint8> VisibilityData;
 };
 
@@ -1229,18 +1229,18 @@ class FLightingAndOcclusion
 public:
 
 	FLinearColor Lighting;
-	FVector UnoccludedSkyVector;
+	FVector3f UnoccludedSkyVector;
 	FLinearColor StationarySkyLighting;
 	float NumSamplesOccluded;
 
 	FLightingAndOcclusion() :
 		Lighting(ForceInit),
-		UnoccludedSkyVector(FVector(0)),
+		UnoccludedSkyVector(FVector3f(0)),
 		StationarySkyLighting(FLinearColor::Black),
 		NumSamplesOccluded(0)
 	{}
 
-	FLightingAndOcclusion(const FLinearColor& InLighting, FVector InUnoccludedSkyVector, const FLinearColor& InStationarySkyLighting, float InNumSamplesOccluded) : 
+	FLightingAndOcclusion(const FLinearColor& InLighting, FVector3f InUnoccludedSkyVector, const FLinearColor& InStationarySkyLighting, float InNumSamplesOccluded) : 
 		Lighting(InLighting),
 		UnoccludedSkyVector(InUnoccludedSkyVector),
 		StationarySkyLighting(InStationarySkyLighting),
@@ -1263,15 +1263,15 @@ class FRefinementElement
 {
 public:
 	FLightingAndOcclusion Lighting;
-	FVector2D Uniforms;
+	FVector2f Uniforms;
 	int32 HitPointIndex;
 
 	FRefinementElement() :
-		Uniforms(FVector2D(0, 0)),
+		Uniforms(FVector2f(0, 0)),
 		HitPointIndex(-1)
 	{}
 
-	FRefinementElement(FLightingAndOcclusion InLighting, FVector2D InUniforms, int32 InHitPointIndex) :
+	FRefinementElement(FLightingAndOcclusion InLighting, FVector2f InUniforms, int32 InHitPointIndex) :
 		Lighting(InLighting),
 		Uniforms(InUniforms),
 		HitPointIndex(InHitPointIndex)
@@ -1421,11 +1421,11 @@ protected:
 class FDirectPhotonEmittingInput
 {
 public:
-	const FBoxSphereBounds& ImportanceBounds;
+	const FBoxSphereBounds3f& ImportanceBounds;
 	const FSceneLightPowerDistribution& LightDistribution;
 	
 	FDirectPhotonEmittingInput(
-		const FBoxSphereBounds& InImportanceBounds,
+		const FBoxSphereBounds3f& InImportanceBounds,
 		const FSceneLightPowerDistribution& InLightDistribution)
 		:
 		ImportanceBounds(InImportanceBounds),
@@ -1510,12 +1510,12 @@ protected:
 class FIndirectPhotonEmittingInput
 {
 public:
-	const FBoxSphereBounds& ImportanceBounds;
+	const FBoxSphereBounds3f& ImportanceBounds;
 	const FSceneLightPowerDistribution& LightDistribution;
 	const TArray<TArray<FIndirectPathRay> >& IndirectPathRays;
 	
 	FIndirectPhotonEmittingInput(
-		const FBoxSphereBounds& InImportanceBounds,
+		const FBoxSphereBounds3f& InImportanceBounds,
 		const FSceneLightPowerDistribution& InLightDistribution,
 		const TArray<TArray<FIndirectPathRay> >& InIndirectPathRays)
 		:
@@ -1830,7 +1830,7 @@ public:
 class FVisibilityMeshGroup
 {
 public:
-	FBox GroupBounds;
+	FBox3f GroupBounds;
 
 	/** Array of all the meshes contained in the group.  These entries index into VisibilityMeshes. */
 	TArray<int32> VisibilityIds;
@@ -1922,8 +1922,8 @@ public:
 		int32 MaterialIndex,
 		int32 UVIndex, 
 		bool bDebugThisMapping, 
-		FVector2D UVBias, 
-		FVector2D UVScale) const;
+		FVector2f UVBias, 
+		FVector2f UVScale) const;
 
 	FStaticLightingAggregateMeshType& GetAggregateMesh()
 	{
@@ -1940,10 +1940,10 @@ private:
 	int32 GetNumUniformHemisphereSamples(int32 BounceNumber) const;
 	int32 GetNumPhotonImportanceHemisphereSamples() const;
 
-	FBoxSphereBounds GetImportanceBounds(bool bClampToScene = true) const;
+	FBoxSphereBounds3f GetImportanceBounds(bool bClampToScene = true) const;
 
 	/** Returns true if the specified position is inside any of the importance volumes. */
-	bool IsPointInImportanceVolume(const FVector4& Position, float Tolerance = 0.0f) const;
+	bool IsPointInImportanceVolume(const FVector4f& Position, float Tolerance = 0.0f) const;
 
 	/** Changes the scene's settings if necessary so that only valid combinations are used */
 	void ValidateSettings(FScene& InScene);
@@ -1968,7 +1968,7 @@ private:
 
 	/** Gathers direct photons and generates indirect photon paths. */
 	void EmitDirectPhotons(
-		const FBoxSphereBounds& ImportanceBounds, 
+		const FBoxSphereBounds3f& ImportanceBounds, 
 		TArray<TArray<FIndirectPathRay> >& IndirectPathRays,
 		TArray<TArray<FIrradiancePhoton>>& IrradiancePhotons);
 
@@ -1985,7 +1985,7 @@ private:
 
 	/** Gathers indirect photons based on the indirect photon paths. */
 	void EmitIndirectPhotons(
-		const FBoxSphereBounds& ImportanceBounds,
+		const FBoxSphereBounds3f& ImportanceBounds,
 		const TArray<TArray<FIndirectPathRay> >& IndirectPathRays, 
 		TArray<TArray<FIrradiancePhoton>>& IrradiancePhotons);
 
@@ -1999,7 +1999,7 @@ private:
 		FIndirectPhotonEmittingOutput& Output);
 
 	/** Iterates through all irradiance photons, searches for nearby direct photons, and marks the irradiance photon has having direct photon influence if necessary. */
-	void MarkIrradiancePhotons(const FBoxSphereBounds& ImportanceBounds, TArray<TArray<FIrradiancePhoton>>& IrradiancePhotons);
+	void MarkIrradiancePhotons(const FBoxSphereBounds3f& ImportanceBounds, TArray<TArray<FIrradiancePhoton>>& IrradiancePhotons);
 
 	/** Entry point for all threads marking irradiance photons. */
 	void MarkIrradiancePhotonsThreadLoop(
@@ -2012,7 +2012,7 @@ private:
 		FIrradianceMarkingWorkRange WorkRange);
 
 	/** Calculates irradiance for photons randomly chosen to precalculate irradiance. */
-	void CalculateIrradiancePhotons(const FBoxSphereBounds& ImportanceBounds, TArray<TArray<FIrradiancePhoton>>& IrradiancePhotons);
+	void CalculateIrradiancePhotons(const FBoxSphereBounds3f& ImportanceBounds, TArray<TArray<FIrradiancePhoton>>& IrradiancePhotons);
 
 	/** Main loop that all threads access to calculate irradiance photons. */
 	void CalculateIrradiancePhotonsThreadLoop(
@@ -2062,7 +2062,7 @@ private:
 	/** Returns true if a photon was found within MaxPhotonSearchDistance. */
 	bool FindAnyNearbyPhoton(
 		const FPhotonOctree& PhotonMap, 
-		const FVector4& SearchPosition, 
+		const FVector4f& SearchPosition, 
 		float MaxPhotonSearchDistance,
 		bool bDebugThisLookup) const;
 
@@ -2075,8 +2075,8 @@ private:
 	 */
 	float FindNearbyPhotonsIterative(
 		const FPhotonOctree& PhotonMap, 
-		const FVector4& SearchPosition, 
-		const FVector4& SearchNormal, 
+		const FVector4f& SearchPosition, 
+		const FVector4f& SearchNormal, 
 		int32 NumPhotonsToFind,
 		float StartPhotonSearchDistance, 
 		float MaxPhotonSearchDistance,
@@ -2087,7 +2087,7 @@ private:
 
 	float FindNearbyPhotonsInVolumeIterative(
 		const FPhotonSegmentOctree& PhotonSegmentMap,
-		const FVector4& SearchPosition,
+		const FVector4f& SearchPosition,
 		int32 NumPhotonsToFind,
 		float StartPhotonSearchDistance,
 		float MaxPhotonSearchDistance,
@@ -2100,8 +2100,8 @@ private:
 	 */
 	float FindNearbyPhotonsSorted(
 		const FPhotonOctree& PhotonMap, 
-		const FVector4& SearchPosition, 
-		const FVector4& SearchNormal, 
+		const FVector4f& SearchPosition, 
+		const FVector4f& SearchNormal, 
 		int32 NumPhotonsToFind, 
 		float MaxPhotonSearchDistance,
 		bool bDebugSearchResults,
@@ -2144,7 +2144,7 @@ private:
 		const FStaticLightingMesh* Mesh,
 		const FMinimalStaticLightingVertex& Vertex,
 		int32 ElementIndex,
-		const FVector4& OutgoingDirection,
+		const FVector4f& OutgoingDirection,
 		bool bDebugThisDensityEstimation) const;
 
 	/** Places volume lighting samples and calculates lighting for them. */
@@ -2168,7 +2168,7 @@ private:
 	void CalculateApproximateDirectLighting(
 		const FStaticLightingVertex& Vertex,
 		float SampleRadius,
-		const TArray<FVector, TInlineAllocator<1>>& VertexOffsets,
+		const TArray<FVector3f, TInlineAllocator<1>>& VertexOffsets,
 		float LightSampleFraction,
 		/** Whether to composite all direct lighting into OutStaticDirectLighting, regardless of static or toggleable. */
 		bool bCompositeAllLights,
@@ -2180,19 +2180,19 @@ private:
 		float& OutToggleableDirectionalLightShadowing) const;
 
 	void GatherVolumeImportancePhotonDirections(
-		const FVector WorldPosition,
-		FVector FirstHemisphereNormal,
-		FVector SecondHemisphereNormal,
-		TArray<FVector4>& FirstHemisphereImportancePhotonDirections,
-		TArray<FVector4>& SecondHemisphereImportancePhotonDirections,
+		const FVector3f WorldPosition,
+		FVector3f FirstHemisphereNormal,
+		FVector3f SecondHemisphereNormal,
+		TArray<FVector4f>& FirstHemisphereImportancePhotonDirections,
+		TArray<FVector4f>& SecondHemisphereImportancePhotonDirections,
 		bool bDebugThisSample) const;
 
 	/** Calculates incident radiance for a given world space position. */
 	void CalculateVolumeSampleIncidentRadiance(
-		const TArray<FVector4>& UniformHemisphereSamples,
-		const TArray<FVector2D>& UniformHemisphereSampleUniforms,
+		const TArray<FVector4f>& UniformHemisphereSamples,
+		const TArray<FVector2f>& UniformHemisphereSampleUniforms,
 		float MaxUnoccludedLength,
-		const TArray<FVector, TInlineAllocator<1>>& VertexOffsets,
+		const TArray<FVector3f, TInlineAllocator<1>>& VertexOffsets,
 		FVolumeLightingSample& LightingSample,
 		float& OutBackfacingHitsFraction,
 		float& OutMinDistanceToSurface,
@@ -2204,7 +2204,7 @@ private:
 	/** Determines visibility cell placement, called once at startup. */
 	void SetupPrecomputedVisibility();
 
-	int32 GetGroupCellIndex(FVector BoxCenter) const;
+	int32 GetGroupCellIndex(FVector3f BoxCenter) const;
 
 	/** Calculates visibility for a given group of cells, called from all threads. */
 	void CalculatePrecomputedVisibility(int32 BucketIndex);
@@ -2215,8 +2215,8 @@ private:
 		FIntVector CellCoordinate, 
 		int32 TreeDepth, 
 		bool bCoveringDebugPosition,
-		const FBox& TopLevelCellBounds, 
-		const TArray<FVector>& VoxelTestPositions,
+		const FBox3f& TopLevelCellBounds, 
+		const TArray<FVector3f>& VoxelTestPositions,
 		const FGuid& IntersectingLevelGuid,
 		TArray<struct FIrradianceBrickBuildData>& OutBrickBuildData);
 
@@ -2224,13 +2224,13 @@ private:
 
 	bool ProcessVolumetricLightmapTaskIfAvailable();
 
-	void GenerateVoxelTestPositions(TArray<FVector>& VoxelTestPositions) const;
+	void GenerateVoxelTestPositions(TArray<FVector3f>& VoxelTestPositions) const;
 
 	void CalculateAdaptiveVolumetricLightmap(int32 TaskIndex);
 
-	bool DoesVoxelIntersectSceneGeometry(const FBox& CellBounds, FGuid& OutIntersectingLevelGuid) const;
+	bool DoesVoxelIntersectSceneGeometry(const FBox3f& CellBounds, FGuid& OutIntersectingLevelGuid) const;
 
-	bool ShouldRefineVoxel(int32 TreeDepth, const FBox& AABB, const TArray<FVector>& VoxelTestPositions, bool bDebugThisVoxel, FGuid& OutIntersectingLevelGuid) const;
+	bool ShouldRefineVoxel(int32 TreeDepth, const FBox3f& AABB, const TArray<FVector3f>& VoxelTestPositions, bool bDebugThisVoxel, FGuid& OutIntersectingLevelGuid) const;
 
 	/** Computes a shadow depth map for a stationary light. */
 	void CalculateStaticShadowDepthMap(FGuid LightGuid);
@@ -2251,13 +2251,13 @@ private:
 	 */
 	bool CalculatePointShadowing(
 		const FStaticLightingMapping* Mapping,
-		const FVector4& WorldSurfacePoint,
+		const FVector4f& WorldSurfacePoint,
 		const FLight* Light,
 		FStaticLightingMappingContext& MappingContext,
 		bool bDebugThisSample) const;
 
 	/** Calculates area shadowing from a light for the given vertex. */
-	FVector2D CalculatePointAreaShadowing(
+	FVector2f CalculatePointAreaShadowing(
 		const FStaticLightingMapping* Mapping,
 		const FStaticLightingVertex& Vertex,
 		int32 ElementIndex,
@@ -2281,12 +2281,12 @@ private:
 		) const;
 
 	/** Returns environment lighting for the given direction. */
-	FLinearColor EvaluateEnvironmentLighting(const FVector4& IncomingDirection) const;
+	FLinearColor EvaluateEnvironmentLighting(const FVector4f& IncomingDirection) const;
 
 	/** Evaluates the incoming sky lighting from the scene. */
-	void EvaluateSkyLighting(const FVector4& IncomingDirection, float PathSolidAngle, bool bShadowed, bool bForDirectLighting, FLinearColor& OutStaticLighting, FLinearColor& OutStationaryLighting) const;
+	void EvaluateSkyLighting(const FVector4f& IncomingDirection, float PathSolidAngle, bool bShadowed, bool bForDirectLighting, FLinearColor& OutStaticLighting, FLinearColor& OutStationaryLighting) const;
 
-	float EvaluateSkyVariance(const FVector4& IncomingDirection, float PathSolidAngle) const;
+	float EvaluateSkyVariance(const FVector4f& IncomingDirection, float PathSolidAngle) const;
 
 	/** Returns a light sample that represents the material attribute specified by MaterialSettings.ViewMaterialAttribute at the intersection. */
 	FGatheredLightSample GetVisualizedMaterialAttribute(const FStaticLightingMapping* Mapping, const FLightRayIntersection& Intersection) const;
@@ -2297,7 +2297,7 @@ private:
 		const FStaticLightingMesh* HitMesh,
 		const FMinimalStaticLightingVertex& Vertex,
 		int32 ElementIndex,
-		const FVector4& OutgoingDirection,
+		const FVector4f& OutgoingDirection,
 		int32 BounceNumber,
 		EHemisphereGatherClassification GatherClassification,
 		FStaticLightingMappingContext& MappingContext,
@@ -2308,8 +2308,8 @@ private:
 		const FFullStaticLightingVertex& Vertex,
 		float SampleRadius,
 		int32 NumRays,
-		const FVector4* WorldPathDirections,
-		const FVector4* TangentPathDirections,
+		const FVector4f* WorldPathDirections,
+		const FVector4f* TangentPathDirections,
 		EFinalGatherRayBiasMode RayBiasMode,
 		FStaticLightingMappingContext& MappingContext,
 		FLightRay* OutLightRays,
@@ -2319,8 +2319,8 @@ private:
 	FLinearColor FinalGatherSample(
 		const FStaticLightingMapping* Mapping,
 		const FFullStaticLightingVertex& Vertex,
-		const FVector4& WorldPathDirection,
-		const FVector4& TangentPathDirection,
+		const FVector4f& WorldPathDirection,
+		const FVector4f& TangentPathDirection,
 		const FLightRay& PathRay,
 		const FLightRayIntersection& RayIntersection,
 		float PathSolidAngle,
@@ -2333,7 +2333,7 @@ private:
 		FLightingCacheGatherInfo& GatherInfo,
 		FFinalGatherInfo& FinalGatherInfo,
 		FFinalGatherHitPoint& HitPoint,
-		FVector& OutUnoccludedSkyVector,
+		FVector3f& OutUnoccludedSkyVector,
 		FLinearColor& OutStationarySkyLighting) const;
 
 	/** 
@@ -2352,10 +2352,10 @@ private:
 		EHemisphereGatherClassification GatherClassification,
 		int32 NumAdaptiveRefinementLevels,
 		float BrightnessThresholdScale,
-		const TArray<FVector4>& UniformHemisphereSamples,
-		const TArray<FVector2D>& UniformHemisphereSampleUniforms,
+		const TArray<FVector4f>& UniformHemisphereSamples,
+		const TArray<FVector2f>& UniformHemisphereSampleUniforms,
 		float MaxUnoccludedLength,
-		const TArray<FVector4>& ImportancePhotonDirections,
+		const TArray<FVector4f>& ImportancePhotonDirections,
 		FStaticLightingMappingContext& MappingContext,
 		FLMRandomStream& RandomStream,
 		FLightingCacheGatherInfo& GatherInfo,
@@ -2366,8 +2366,8 @@ private:
 	void CalculateIrradianceGradients(
 		int32 BounceNumber,
 		const FLightingCacheGatherInfo& GatherInfo,
-		FVector4& RotationalGradient,
-		FVector4& TranslationalGradient) const;
+		FVector4f& RotationalGradient,
+		FVector4f& TranslationalGradient) const;
 
 	/** 
 	 * Interpolates incoming radiance from the lighting cache if possible,
@@ -2392,9 +2392,9 @@ private:
 
 	/** Traces a ray to the corner of a texel. */
 	void TraceToTexelCorner(
-		const FVector4& TexelCenterOffset, 
+		const FVector4f& TexelCenterOffset, 
 		const FFullStaticLightingVertex& FullVertex, 
-		FVector2D CornerSigns,
+		FVector2f CornerSigns,
 		float TexelRadius, 
 		FStaticLightingMappingContext& MappingContext, 
 		FLightRayIntersection& Intersection,
@@ -2601,7 +2601,7 @@ private:
 	volatile int32 NumVolumeSampleTasksOutstanding;
 	volatile int32 bShouldExportVolumeSampleData;
 	/** Bounds that VolumeLightingSamples were generated in. */
-	FBoxSphereBounds VolumeBounds;
+	FBoxSphereBounds3f VolumeBounds;
 	/** Octree used for interpolating the volume lighting samples if DynamicObjectSettings.bVisualizeVolumeLightInterpolation is true. */
 	FVolumeLightingInterpolationOctree VolumeLightingInterpolationOctree;
 	/** Map from Level Guid to array of volume lighting samples generated. */
@@ -2670,20 +2670,20 @@ private:
 	float MaxRayDistance;
 
 	/** Cached direction samples for hemisphere gathers. */
-	TArray<FVector4> CachedHemisphereSamples;
+	TArray<FVector4f> CachedHemisphereSamples;
 
 	/** Length of all the hemisphere samples averaged, which is also the max length that a bent normal can be. */
 	float CachedSamplesMaxUnoccludedLength;
 
-	TArray<FVector2D> CachedHemisphereSampleUniforms;
+	TArray<FVector2f> CachedHemisphereSampleUniforms;
 
-	TArray<FVector4> CachedHemisphereSamplesForRadiosity[3];
-	TArray<FVector2D> CachedHemisphereSamplesForRadiosityUniforms[3];
+	TArray<FVector4f> CachedHemisphereSamplesForRadiosity[3];
+	TArray<FVector2f> CachedHemisphereSamplesForRadiosityUniforms[3];
 
-	TArray<FVector4> CachedVolumetricLightmapUniformHemisphereSamples;
-	TArray<FVector2D> CachedVolumetricLightmapUniformHemisphereSampleUniforms;
+	TArray<FVector4f> CachedVolumetricLightmapUniformHemisphereSamples;
+	TArray<FVector2f> CachedVolumetricLightmapUniformHemisphereSampleUniforms;
 	float CachedVolumetricLightmapMaxUnoccludedLength;
-	TArray<FVector, TInlineAllocator<1>> CachedVolumetricLightmapVertexOffsets;
+	TArray<FVector3f, TInlineAllocator<1>> CachedVolumetricLightmapVertexOffsets;
 
 	/** The aggregate mesh used for raytracing. */
 	FStaticLightingAggregateMeshType* AggregateMesh;
@@ -2711,7 +2711,7 @@ private:
 	int32 GroupVisibilityGridSizeZ;
 
 	/** World space bounding box of GroupGrid. */
-	FBox VisibilityGridBounds;
+	FBox3f VisibilityGridBounds;
 
 	/** Grid of indices into VisibilityGroups. */
 	TArray<int32> GroupGrid;
@@ -2742,8 +2742,8 @@ private:
 	int32 VolumeSizeX;
 	int32 VolumeSizeY;
 	int32 VolumeSizeZ;
-	FVector4::FReal DistanceFieldVoxelSize;
-	FBox DistanceFieldVolumeBounds;
+	FVector4f::FReal DistanceFieldVoxelSize;
+	FBox3f DistanceFieldVolumeBounds;
 	TArray<FColor> VolumeDistanceField;
 
 	/** */
@@ -2818,7 +2818,7 @@ private:
  * @param Light - The light to classify.
  * @return true if the light is behind the triangle.
  */
-extern bool IsLightBehindSurface(const FVector4& TrianglePoint, const FVector4& TriangleNormal, const FLight* Light);
+extern bool IsLightBehindSurface(const FVector4f& TrianglePoint, const FVector4f& TriangleNormal, const FLight* Light);
 
 /**
  * Culls lights that are behind a triangle.
@@ -2828,7 +2828,7 @@ extern bool IsLightBehindSurface(const FVector4& TrianglePoint, const FVector4&
  * @param Lights - The lights to cull.
  * @return A map from Lights index to a boolean which is true if the light is in front of the triangle.
  */
-extern TBitArray<> CullBackfacingLights(bool bTwoSidedMaterial, const FVector4& TrianglePoint, const FVector4& TriangleNormal, const TArray<FLight*>& Lights);
+extern TBitArray<> CullBackfacingLights(bool bTwoSidedMaterial, const FVector4f& TrianglePoint, const FVector4f& TriangleNormal, const TArray<FLight*>& Lights);
 
 #include "LightingSystem.inl"
 
diff --git a/Engine/Source/Programs/UnrealLightmass/Private/Lighting/LightingSystem.inl b/Engine/Source/Programs/UnrealLightmass/Private/Lighting/LightingSystem.inl
--- a/Engine/Source/Programs/UnrealLightmass/Private/Lighting/LightingSystem.inl
+++ b/Engine/Source/Programs/UnrealLightmass/Private/Lighting/LightingSystem.inl
@@ -7,9 +7,9 @@ TGatheredLightSample<SHOrder> FGatheredLightSampleUtil::AmbientLight(const FLine
 	Result.SHVector.AddAmbient(Color);
 
 	// Compute SHCorrection as if all the lighting was coming in along the normal
-	FVector4 TangentDirection(0, 0, 1);
+	FVector4f TangentDirection(0, 0, 1);
 
-	FSHVector2 SH = FSHVector2::SHBasisFunction(TangentDirection);
+	FSHVector2 SH = FSHVector2::SHBasisFunction(FVector4(TangentDirection));
 	Result.SHCorrection = Color.GetLuminance() * (0.282095f * SH.V[0] + 0.325735f * SH.V[2]);
 
 	Result.IncidentLighting = Color;
@@ -20,7 +20,7 @@ TGatheredLightSample<SHOrder> FGatheredLightSampleUtil::AmbientLight(const FLine
 }
 
 template<int32 SHOrder>
-TGatheredLightSample<SHOrder> FGatheredLightSampleUtil::PointLightWorldSpace(const FLinearColor& Color, const FVector4& TangentDirection, const FVector4& WorldDirection)
+TGatheredLightSample<SHOrder> FGatheredLightSampleUtil::PointLightWorldSpace(const FLinearColor& Color, const FVector4f& TangentDirection, const FVector4f& WorldDirection)
 {
 	TGatheredLightSample<SHOrder> Result;
 
@@ -28,7 +28,7 @@ TGatheredLightSample<SHOrder> FGatheredLightSampleUtil::PointLightWorldSpace(con
 	{
 		Result.SHVector.AddIncomingRadiance(Color, 1, WorldDirection);
 
-		FSHVector2 SH = FSHVector2::SHBasisFunction(TangentDirection);
+		FSHVector2 SH = FSHVector2::SHBasisFunction(FVector4(TangentDirection));
 		// Evaluate lighting along the smoothed vertex normal direction, so that later we can guarantee an SH intensity of 1 along the normal
 		// These scaling coefficients are SHBasisFunction and CalcDiffuseTransferSH baked down
 		// 0.325735f = 0.488603f from SHBasisFunction * 2/3 from CalcDiffuseTransferSH
@@ -93,7 +93,7 @@ template<int32 SHOrder>
 void FStaticLightingSystem::CalculateApproximateDirectLighting(
 	const FStaticLightingVertex& Vertex,
 	float SampleRadius,
-	const TArray<FVector, TInlineAllocator<1>>& VertexOffsets,
+	const TArray<FVector3f, TInlineAllocator<1>>& VertexOffsets,
 	float LightSampleFraction,
 	bool bCompositeAllLights,
 	bool bCalculateForIndirectLighting,
@@ -109,7 +109,7 @@ void FStaticLightingSystem::CalculateApproximateDirectLighting(
 	{
 		const FLight* Light = Lights[LightIndex];
 
-		if (Light->AffectsBounds(FBoxSphereBounds(FSphere(Vertex.WorldPosition, SampleRadius))))
+		if (Light->AffectsBounds(FBoxSphereBounds3f(FSphere3f(Vertex.WorldPosition, SampleRadius))))
 		{
 			FLinearColor LightIntensity(0, 0, 0, 0);
 
@@ -137,11 +137,11 @@ void FStaticLightingSystem::CalculateApproximateDirectLighting(
 					Light->ValidateSurfaceSample(Vertex.WorldPosition, CurrentSample);
 
 					// Construct a line segment between the light and the volume point.
-					const FVector4 LightVector = CurrentSample.Position - Vertex.WorldPosition;
+					const FVector4f LightVector = CurrentSample.Position - Vertex.WorldPosition;
 
-					FVector4 NormalForOffset = Vertex.WorldTangentZ;
+					FVector4f NormalForOffset = Vertex.WorldTangentZ;
 
-					const FVector4 StartOffset = LightVector.GetSafeNormal() * SceneConstants.VisibilityRayOffsetDistance
+					const FVector4f StartOffset = LightVector.GetSafeNormal() * SceneConstants.VisibilityRayOffsetDistance
 						+ NormalForOffset * SampleRadius * SceneConstants.VisibilityNormalOffsetSampleRadiusScale;
 
 					const FLightRay LightRay(
@@ -190,11 +190,11 @@ void FStaticLightingSystem::CalculateApproximateDirectLighting(
 
 			{
 				// Calculate the direction from the vertex to the light.
-				const FVector4 WorldLightVector = Light->GetDirectLightingDirection(Vertex.WorldPosition, Vertex.WorldTangentZ);
+				const FVector4f WorldLightVector = Light->GetDirectLightingDirection(Vertex.WorldPosition, Vertex.WorldTangentZ);
 
 				// Transform the light vector to tangent space.
-				const FVector4 TangentLightVector =
-					FVector4(
+				const FVector4f TangentLightVector =
+					FVector4f(
 					Dot3(WorldLightVector, Vertex.WorldTangentX),
 					Dot3(WorldLightVector, Vertex.WorldTangentY),
 					Dot3(WorldLightVector, Vertex.WorldTangentZ),
diff --git a/Engine/Source/Programs/UnrealLightmass/Private/Lighting/LightmapData.cpp b/Engine/Source/Programs/UnrealLightmass/Private/Lighting/LightmapData.cpp
--- a/Engine/Source/Programs/UnrealLightmass/Private/Lighting/LightmapData.cpp
+++ b/Engine/Source/Programs/UnrealLightmass/Private/Lighting/LightmapData.cpp
@@ -263,9 +263,9 @@ namespace Lightmass
 #endif
 			DestCoefficients.Coverage = SourceSample.bIsMapped ? 255 : 0;
 
-			const FVector BentNormal(SourceSample.SkyOcclusion[0], SourceSample.SkyOcclusion[1], SourceSample.SkyOcclusion[2]);
+			const FVector3f BentNormal(SourceSample.SkyOcclusion[0], SourceSample.SkyOcclusion[1], SourceSample.SkyOcclusion[2]);
 			const float BentNormalLength = BentNormal.Size();
-			const FVector NormalizedBentNormal = BentNormal.GetSafeNormal() * FVector(.5f) + FVector(.5f);
+			const FVector3f NormalizedBentNormal = BentNormal.GetSafeNormal() * FVector3f(.5f) + FVector3f(.5f);
 
 			DestCoefficients.SkyOcclusion[0] = (uint8)FMath::Clamp<int32>( FMath::RoundToInt( NormalizedBentNormal[0] * 255.0f ), 0, 255 );
 			DestCoefficients.SkyOcclusion[1] = (uint8)FMath::Clamp<int32>( FMath::RoundToInt( NormalizedBentNormal[1] * 255.0f ), 0, 255 );
diff --git a/Engine/Source/Programs/UnrealLightmass/Private/Lighting/MonteCarlo.cpp b/Engine/Source/Programs/UnrealLightmass/Private/Lighting/MonteCarlo.cpp
--- a/Engine/Source/Programs/UnrealLightmass/Private/Lighting/MonteCarlo.cpp
+++ b/Engine/Source/Programs/UnrealLightmass/Private/Lighting/MonteCarlo.cpp
@@ -6,14 +6,14 @@ namespace Lightmass
 {
 
 /** Generates valid X and Y axes of a coordinate system, given the Z axis. */
-void GenerateCoordinateSystem(const FVector4& ZAxis, FVector4& XAxis, FVector4& YAxis)
+void GenerateCoordinateSystem(const FVector4f& ZAxis, FVector4f& XAxis, FVector4f& YAxis)
 {
 	// Use the vector perpendicular to ZAxis and the Y axis as the XAxis
-	const FVector4 XAxisCandidate = ZAxis ^ FVector4(0,1,0);
+	const FVector4f XAxisCandidate = ZAxis ^ FVector4f(0,1,0);
 	if (XAxisCandidate.SizeSquared3() < KINDA_SMALL_NUMBER)
 	{
 		// The vector was nearly equal to the Y axis, use the X axis instead
-		XAxis = (ZAxis ^ FVector4(1,0,0)).GetUnsafeNormal3();
+		XAxis = (ZAxis ^ FVector4f(1,0,0)).GetUnsafeNormal3();
 	}
 	else
 	{
@@ -25,18 +25,18 @@ void GenerateCoordinateSystem(const FVector4& ZAxis, FVector4& XAxis, FVector4&
 }
 
 /** Generates valid X and Y axes of a coordinate system, given the Z axis. */
-void GenerateCoordinateSystem2(const FVector4& ZAxis, FVector4& XAxis, FVector4& YAxis)
+void GenerateCoordinateSystem2(const FVector4f& ZAxis, FVector4f& XAxis, FVector4f& YAxis)
 {
 	// This implementation is based off of the one from 'Physically Based Rendering'
 	if (FMath::Abs(ZAxis.X) > FMath::Abs(ZAxis.Y))
  	{
 		const float InverseLength = FMath::InvSqrt(ZAxis.X * ZAxis.X + ZAxis.Z * ZAxis.Z);
-		XAxis = FVector4(-ZAxis.Z * InverseLength, 0.0f, ZAxis.X * InverseLength);
+		XAxis = FVector4f(-ZAxis.Z * InverseLength, 0.0f, ZAxis.X * InverseLength);
  	}
  	else
  	{
 		const float InverseLength = FMath::InvSqrt(ZAxis.Y * ZAxis.Y + ZAxis.Z * ZAxis.Z);
-		XAxis = FVector4(0.0f, ZAxis.Z * InverseLength, -ZAxis.Y * InverseLength);
+		XAxis = FVector4f(0.0f, ZAxis.Z * InverseLength, -ZAxis.Y * InverseLength);
  	}
 
 	YAxis = ZAxis ^ XAxis;
@@ -47,15 +47,15 @@ void GenerateCoordinateSystem2(const FVector4& ZAxis, FVector4& XAxis, FVector4&
 }
 
 /** Generates a pseudo-random unit vector, uniformly distributed over all directions. */
-FVector4 GetUnitVector(FLMRandomStream& RandomStream)
+FVector4f GetUnitVector(FLMRandomStream& RandomStream)
 {
 	return GetUnitPosition(RandomStream).GetUnsafeNormal3();
 }
 
 /** Generates a pseudo-random position inside the unit sphere, uniformly distributed over the volume of the sphere. */
-FVector4 GetUnitPosition(FLMRandomStream& RandomStream)
+FVector4f GetUnitPosition(FLMRandomStream& RandomStream)
 {
-	FVector4 Result;
+	FVector4f Result;
 	// Use rejection sampling to generate a valid sample
 	do
 	{
@@ -70,7 +70,7 @@ FVector4 GetUnitPosition(FLMRandomStream& RandomStream)
  * Generates a pseudo-random unit vector in the Z > 0 hemisphere whose PDF == 1 / (2 * PI) in solid angles,
  * Or sin(theta) / (2 * PI) in hemispherical coordinates, which is a uniform distribution over the area of the hemisphere.
  */
-FVector4 GetUniformHemisphereVector(FLMRandomStream& RandomStream, float MaxTheta)
+FVector4f GetUniformHemisphereVector(FLMRandomStream& RandomStream, float MaxTheta)
 {
 	const float Theta = FMath::Min(FMath::Acos(RandomStream.GetFraction()), MaxTheta - DELTA);
 	const float Phi = 2.0f * (float)PI * RandomStream.GetFraction();
@@ -78,14 +78,14 @@ FVector4 GetUniformHemisphereVector(FLMRandomStream& RandomStream, float MaxThet
 	checkSlow(Phi >= 0 && Phi <= 2.0f * (float)PI);
 	const float SinTheta = FMath::Sin(Theta);
 	// Convert to Cartesian
-	return FVector4(FMath::Cos(Phi) * SinTheta, FMath::Sin(Phi) * SinTheta, FMath::Cos(Theta));
+	return FVector4f(FMath::Cos(Phi) * SinTheta, FMath::Sin(Phi) * SinTheta, FMath::Cos(Theta));
 }
 
 /** 
  * Generates a pseudo-random unit vector in the Z > 0 hemisphere whose PDF == cos(theta) / PI in solid angles,
  * Which is sin(theta)cos(theta) / PI in hemispherical coordinates.
  */
-FVector4 GetCosineHemisphereVector(FLMRandomStream& RandomStream, float MaxTheta)
+FVector4f GetCosineHemisphereVector(FLMRandomStream& RandomStream, float MaxTheta)
 {
 	const float Theta = FMath::Min(FMath::Acos(FMath::Sqrt(RandomStream.GetFraction())), MaxTheta - DELTA);
 	const float Phi = 2.0f * (float)PI * RandomStream.GetFraction();
@@ -93,7 +93,7 @@ FVector4 GetCosineHemisphereVector(FLMRandomStream& RandomStream, float MaxTheta
 	checkSlow(Phi >= 0 && Phi <= 2.0f * (float)PI);
 	const float SinTheta = FMath::Sin(Theta);
 	// Convert to Cartesian
-	return FVector4(FMath::Cos(Phi) * SinTheta, FMath::Sin(Phi) * SinTheta, FMath::Cos(Theta));
+	return FVector4f(FMath::Cos(Phi) * SinTheta, FMath::Sin(Phi) * SinTheta, FMath::Cos(Theta));
 }
 
 /** 
@@ -101,12 +101,12 @@ FVector4 GetCosineHemisphereVector(FLMRandomStream& RandomStream, float MaxTheta
  * Whose PDF == (SpecularPower + 1) / (2.0f * PI) * cos(Alpha) ^ SpecularPower in solid angles,
  * Where Alpha is the angle between the perfect specular direction and the outgoing direction.
  */
-FVector4 GetModifiedPhongSpecularVector(FLMRandomStream& RandomStream, const FVector4& TangentSpecularDirection, float SpecularPower)
+FVector4f GetModifiedPhongSpecularVector(FLMRandomStream& RandomStream, const FVector4f& TangentSpecularDirection, float SpecularPower)
 {
 	checkSlow(TangentSpecularDirection.Z >= 0.0f);
 	checkSlow(SpecularPower > 0.0f);
 
-	FVector4 GeneratedTangentVector;
+	FVector4f GeneratedTangentVector;
 	do
 	{
 		// Generate hemispherical coordinates in the local frame of the perfect specular direction
@@ -116,26 +116,26 @@ FVector4 GetModifiedPhongSpecularVector(FLMRandomStream& RandomStream, const FVe
 		
 		// Convert to Cartesian, still in the coordinate space of the perfect specular direction
 		const float SinTheta = FMath::Sin(Alpha);
-		const FVector4 GeneratedSpecularTangentVector(FMath::Cos(Phi) * SinTheta, FMath::Sin(Phi) * SinTheta, FMath::Cos(Alpha));
+		const FVector4f GeneratedSpecularTangentVector(FMath::Cos(Phi) * SinTheta, FMath::Sin(Phi) * SinTheta, FMath::Cos(Alpha));
 
 		// Generate the X and Y axes of the coordinate space whose Z is the perfect specular direction
-		FVector4 SpecularTangentX = (TangentSpecularDirection ^ FVector4(0,1,0)).GetUnsafeNormal3();
+		FVector4f SpecularTangentX = (TangentSpecularDirection ^ FVector4f(0,1,0)).GetUnsafeNormal3();
 		if (SpecularTangentX.SizeSquared3() < KINDA_SMALL_NUMBER)
 		{
 			// The specular direction was nearly equal to the Y axis, use the X axis instead
-			SpecularTangentX = (TangentSpecularDirection ^ FVector4(1,0,0)).GetUnsafeNormal3();
+			SpecularTangentX = (TangentSpecularDirection ^ FVector4f(1,0,0)).GetUnsafeNormal3();
 		}
 		else
 		{
 			SpecularTangentX = SpecularTangentX.GetUnsafeNormal3();
 		}
-		const FVector4 SpecularTangentY = TangentSpecularDirection ^ SpecularTangentX;
+		const FVector4f SpecularTangentY = TangentSpecularDirection ^ SpecularTangentX;
 
 		// Rotate the generated coordinates into the local frame of the tangent space normal (0,0,1)
-		const FVector4 SpecularTangentRow0(SpecularTangentX.X, SpecularTangentY.X, TangentSpecularDirection.X);
-		const FVector4 SpecularTangentRow1(SpecularTangentX.Y, SpecularTangentY.Y, TangentSpecularDirection.Y);
-		const FVector4 SpecularTangentRow2(SpecularTangentX.Z, SpecularTangentY.Z, TangentSpecularDirection.Z);
-		GeneratedTangentVector = FVector4(
+		const FVector4f SpecularTangentRow0(SpecularTangentX.X, SpecularTangentY.X, TangentSpecularDirection.X);
+		const FVector4f SpecularTangentRow1(SpecularTangentX.Y, SpecularTangentY.Y, TangentSpecularDirection.Y);
+		const FVector4f SpecularTangentRow2(SpecularTangentX.Z, SpecularTangentY.Z, TangentSpecularDirection.Z);
+		GeneratedTangentVector = FVector4f(
 			Dot3(SpecularTangentRow0, GeneratedSpecularTangentVector),
 			Dot3(SpecularTangentRow1, GeneratedSpecularTangentVector),
 			Dot3(SpecularTangentRow2, GeneratedSpecularTangentVector)
@@ -151,18 +151,18 @@ FVector4 GetModifiedPhongSpecularVector(FLMRandomStream& RandomStream, const FVe
  * Generates a pseudo-random position within a unit disk,
  * Whose PDF == 1 / PI, which is a uniform distribution over the area of the disk.
  */
-FVector2D GetUniformUnitDiskPosition(FLMRandomStream& RandomStream)
+FVector2f GetUniformUnitDiskPosition(FLMRandomStream& RandomStream)
 {
 	const float Theta = 2.0f * (float)PI * RandomStream.GetFraction();
 	const float Radius = FMath::Sqrt(RandomStream.GetFraction());
-	return FVector2D(Radius * FMath::Cos(Theta), Radius * FMath::Sin(Theta));
+	return FVector2f(Radius * FMath::Cos(Theta), Radius * FMath::Sin(Theta));
 }
 
 /** 
  * Generates a pseudo-random direction within a cone,
  * Whose PDF == 1 / (2 * PI * (1 - CosMaxConeTheta)), which is a uniform distribution over the directions in the cone. 
  */
-FVector4 UniformSampleCone(FLMRandomStream& RandomStream, float CosMaxConeTheta, const FVector4& XAxis, const FVector4& YAxis, const FVector4& ZAxis)
+FVector4f UniformSampleCone(FLMRandomStream& RandomStream, float CosMaxConeTheta, const FVector4f& XAxis, const FVector4f& YAxis, const FVector4f& ZAxis)
 {
 	checkSlow(CosMaxConeTheta >= 0.0f && CosMaxConeTheta <= 1.0f);
 	const float CosTheta = FMath::Lerp(CosMaxConeTheta, 1.0f, RandomStream.GetFraction());
@@ -171,7 +171,7 @@ FVector4 UniformSampleCone(FLMRandomStream& RandomStream, float CosMaxConeTheta,
 	return FMath::Cos(Phi) * SinTheta * XAxis + FMath::Sin(Phi) * SinTheta * YAxis + CosTheta * ZAxis;
 }
 
-FVector4 UniformSampleCone(float CosMaxConeTheta, const FVector4& XAxis, const FVector4& YAxis, const FVector4& ZAxis, float Uniform1, float Uniform2)
+FVector4f UniformSampleCone(float CosMaxConeTheta, const FVector4f& XAxis, const FVector4f& YAxis, const FVector4f& ZAxis, float Uniform1, float Uniform2)
 {
 	checkSlow(CosMaxConeTheta >= 0.0f && CosMaxConeTheta <= 1.0f);
 	const float CosTheta = FMath::Lerp(CosMaxConeTheta, 1.0f, Uniform1);
@@ -187,17 +187,17 @@ float UniformConePDF(float CosMaxConeTheta)
 	return 1.0f / (2.0f * (float)PI * (1.0f - CosMaxConeTheta));
 }
 
-FVector4 UniformSampleHemisphere(float Uniform1, float Uniform2)
+FVector4f UniformSampleHemisphere(float Uniform1, float Uniform2)
 {
 	const float R = FMath::Sqrt(1.0f - Uniform1 * Uniform1);
 	const float Phi = 2.0f * (float)PI * Uniform2;
 
 	// Convert to Cartesian
-	return FVector4(FMath::Cos(Phi) * R, FMath::Sin(Phi) * R, Uniform1);
+	return FVector4f(FMath::Cos(Phi) * R, FMath::Sin(Phi) * R, Uniform1);
 }
 
 /** Generates unit length, stratified and uniformly distributed direction samples in a hemisphere. */
-void GenerateStratifiedUniformHemisphereSamples(int32 NumThetaSteps, int32 NumPhiSteps, FLMRandomStream& RandomStream, TArray<FVector4>& Samples, TArray<FVector2D>& Uniforms)
+void GenerateStratifiedUniformHemisphereSamples(int32 NumThetaSteps, int32 NumPhiSteps, FLMRandomStream& RandomStream, TArray<FVector4f>& Samples, TArray<FVector2f>& Uniforms)
 {
 	Samples.Empty(NumThetaSteps * NumPhiSteps);
 	for (int32 ThetaIndex = 0; ThetaIndex < NumThetaSteps; ThetaIndex++)
@@ -214,13 +214,13 @@ void GenerateStratifiedUniformHemisphereSamples(int32 NumThetaSteps, int32 NumPh
 
 			const float Phi = 2.0f * (float)PI * Fraction2;
 			// Convert to Cartesian
-			Samples.Add(FVector4(FMath::Cos(Phi) * R, FMath::Sin(Phi) * R, Fraction1));
-			Uniforms.Add(FVector2D(Fraction1, Fraction2));
+			Samples.Add(FVector4f(FMath::Cos(Phi) * R, FMath::Sin(Phi) * R, Fraction1));
+			Uniforms.Add(FVector2f(Fraction1, Fraction2));
 		}
 	}
 }
 
-void GenerateStratifiedCosineHemisphereSamples(int32 NumThetaSteps, int32 NumPhiSteps, FLMRandomStream& RandomStream, TArray<FVector4>& Samples)
+void GenerateStratifiedCosineHemisphereSamples(int32 NumThetaSteps, int32 NumPhiSteps, FLMRandomStream& RandomStream, TArray<FVector4f>& Samples)
 {
 	Samples.Empty(NumThetaSteps * NumPhiSteps);
 
@@ -240,7 +240,7 @@ void GenerateStratifiedCosineHemisphereSamples(int32 NumThetaSteps, int32 NumPhi
 			checkSlow(Phi >= 0 && Phi <= 2.0f * (float)PI);
 			const float SinTheta = FMath::Sin(Theta);
 			// Convert to Cartesian
-			Samples.Add(FVector4(FMath::Cos(Phi) * SinTheta, FMath::Sin(Phi) * SinTheta, FMath::Cos(Theta)));
+			Samples.Add(FVector4f(FMath::Cos(Phi) * SinTheta, FMath::Sin(Phi) * SinTheta, FMath::Cos(Theta)));
 		}
 	}
 }
diff --git a/Engine/Source/Programs/UnrealLightmass/Private/Lighting/MonteCarlo.h b/Engine/Source/Programs/UnrealLightmass/Private/Lighting/MonteCarlo.h
--- a/Engine/Source/Programs/UnrealLightmass/Private/Lighting/MonteCarlo.h
+++ b/Engine/Source/Programs/UnrealLightmass/Private/Lighting/MonteCarlo.h
@@ -8,56 +8,56 @@
 namespace Lightmass
 {
 	/** Generates valid X and Y axes of a coordinate system, given the Z axis. */
-	void GenerateCoordinateSystem(const FVector4& ZAxis, FVector4& XAxis, FVector4& YAxis);
+	void GenerateCoordinateSystem(const FVector4f& ZAxis, FVector4f& XAxis, FVector4f& YAxis);
 		
 	/** Generates a pseudo-random unit vector, uniformly distributed over all directions. */
-	FVector4 GetUnitVector(FLMRandomStream& RandomStream);
+	FVector4f GetUnitVector(FLMRandomStream& RandomStream);
 
 	/** Generates a pseudo-random position inside the unit sphere, uniformly distributed over the volume of the sphere. */
-	FVector4 GetUnitPosition(FLMRandomStream& RandomStream);
+	FVector4f GetUnitPosition(FLMRandomStream& RandomStream);
 
 	/** 
 	 * Generates a pseudo-random unit vector in the Z > 0 hemisphere whose PDF == 1 / (2 * PI) in solid angles,
 	 * Or sin(theta) / (2 * PI) in hemispherical coordinates, which is a uniform distribution over the area of the hemisphere.
 	 */
-	FVector4 GetUniformHemisphereVector(FLMRandomStream& RandomStream, float MaxTheta = (float)HALF_PI);
+	FVector4f GetUniformHemisphereVector(FLMRandomStream& RandomStream, float MaxTheta = (float)HALF_PI);
 
 	/** 
 	 * Generates a pseudo-random unit vector in the Z > 0 hemisphere whose PDF == cos(theta) / PI in solid angles,
 	 * Which is sin(theta)cos(theta) / PI in hemispherical coordinates.
 	 */
-	FVector4 GetCosineHemisphereVector(FLMRandomStream& RandomStream, float MaxTheta = (float)HALF_PI);
+	FVector4f GetCosineHemisphereVector(FLMRandomStream& RandomStream, float MaxTheta = (float)HALF_PI);
 
 	/** 
 	 * Generates a pseudo-random unit vector in the Z > 0 hemisphere,
 	 * Whose PDF == (SpecularPower + 1) / (2.0f * PI) * cos(Alpha) ^ SpecularPower in solid angles,
 	 * Where Alpha is the angle between the perfect specular direction and the outgoing direction.
 	 */
-	FVector4 GetModifiedPhongSpecularVector(FLMRandomStream& RandomStream, const FVector4& TangentSpecularDirection, float SpecularPower);
+	FVector4f GetModifiedPhongSpecularVector(FLMRandomStream& RandomStream, const FVector4f& TangentSpecularDirection, float SpecularPower);
 
 	/** 
 	 * Generates a pseudo-random position within a unit disk,
 	 * Whose PDF == 1 / PI, which is a uniform distribution over the area of the disk.
 	 */
-	FVector2D GetUniformUnitDiskPosition(FLMRandomStream& RandomStream);
+	FVector2f GetUniformUnitDiskPosition(FLMRandomStream& RandomStream);
 
 	/** 
 	 * Generates a pseudo-random direction within a cone,
 	 * Whose PDF == 1 / (2 * PI * (1 - CosMaxConeTheta)), which is a uniform distribution over the directions in the cone. 
 	 */
-	FVector4 UniformSampleCone(FLMRandomStream& RandomStream, float CosMaxConeTheta, const FVector4& XAxis, const FVector4& YAxis, const FVector4& ZAxis);
+	FVector4f UniformSampleCone(FLMRandomStream& RandomStream, float CosMaxConeTheta, const FVector4f& XAxis, const FVector4f& YAxis, const FVector4f& ZAxis);
 
-	FVector4 UniformSampleCone(float CosMaxConeTheta, const FVector4& XAxis, const FVector4& YAxis, const FVector4& ZAxis, float Uniform1, float Uniform2);
+	FVector4f UniformSampleCone(float CosMaxConeTheta, const FVector4f& XAxis, const FVector4f& YAxis, const FVector4f& ZAxis, float Uniform1, float Uniform2);
 
 	/** Calculates the PDF for a sample generated by UniformSampleCone */
 	float UniformConePDF(float CosMaxConeTheta);
 
-	FVector4 UniformSampleHemisphere(float Uniform1, float Uniform2);
+	FVector4f UniformSampleHemisphere(float Uniform1, float Uniform2);
 
 	/** Generates unit length, stratified and uniformly distributed direction samples in a hemisphere. */
-	void GenerateStratifiedUniformHemisphereSamples(int32 NumThetaSteps, int32 NumPhiSteps, FLMRandomStream& RandomStream, TArray<FVector4>& Samples, TArray<FVector2D>& Uniforms);
+	void GenerateStratifiedUniformHemisphereSamples(int32 NumThetaSteps, int32 NumPhiSteps, FLMRandomStream& RandomStream, TArray<FVector4f>& Samples, TArray<FVector2f>& Uniforms);
 
-	void GenerateStratifiedCosineHemisphereSamples(int32 NumThetaSteps, int32 NumPhiSteps, FLMRandomStream& RandomStream, TArray<FVector4>& Samples);
+	void GenerateStratifiedCosineHemisphereSamples(int32 NumThetaSteps, int32 NumPhiSteps, FLMRandomStream& RandomStream, TArray<FVector4f>& Samples);
 
 	/** 
 	 * Multiple importance sampling power heuristic of two functions with a power of two. 
diff --git a/Engine/Source/Programs/UnrealLightmass/Private/Lighting/PhotonMapping.cpp b/Engine/Source/Programs/UnrealLightmass/Private/Lighting/PhotonMapping.cpp
--- a/Engine/Source/Programs/UnrealLightmass/Private/Lighting/PhotonMapping.cpp
+++ b/Engine/Source/Programs/UnrealLightmass/Private/Lighting/PhotonMapping.cpp
@@ -24,8 +24,8 @@ static const int32 NumPhotonWorkRanges = 256;
 /** Sets up photon mapping settings. */
 void FStaticLightingSystem::InitializePhotonSettings()
 {
-	const FBoxSphereBounds SceneBounds = FBoxSphereBounds(AggregateMesh->GetBounds());
-	const FBoxSphereBounds ImportanceBounds = GetImportanceBounds();
+	const FBoxSphereBounds3f SceneBounds = FBoxSphereBounds3f(AggregateMesh->GetBounds());
+	const FBoxSphereBounds3f ImportanceBounds = GetImportanceBounds();
 
 	// Get direct photon counts from each light
 #if !LIGHTMASS_DO_PROCESSING
@@ -123,8 +123,8 @@ void FStaticLightingSystem::InitializePhotonSettings()
 /** Emits photons, builds data structures to accelerate photon map lookups, and does any other photon preprocessing required. */
 void FStaticLightingSystem::EmitPhotons()
 {
-	const FBoxSphereBounds SceneSphereBounds = FBoxSphereBounds(AggregateMesh->GetBounds());
-	FBoxSphereBounds ImportanceVolumeBounds = GetImportanceBounds();
+	const FBoxSphereBounds3f SceneSphereBounds = FBoxSphereBounds3f(AggregateMesh->GetBounds());
+	FBoxSphereBounds3f ImportanceVolumeBounds = GetImportanceBounds();
 	if (ImportanceVolumeBounds.SphereRadius < DELTA)
 	{
 		ImportanceVolumeBounds = SceneSphereBounds;
@@ -200,7 +200,7 @@ void FStaticLightingSystem::EmitPhotons()
 
 /** Emits direct photons and generates indirect photon paths. */
 void FStaticLightingSystem::EmitDirectPhotons(
-	const FBoxSphereBounds& ImportanceBounds, 
+	const FBoxSphereBounds3f& ImportanceBounds, 
 	TArray<TArray<FIndirectPathRay> >& IndirectPathRays,
 	TArray<TArray<FIrradiancePhoton>>& IrradiancePhotons)
 {
@@ -527,8 +527,8 @@ void FStaticLightingSystem::EmitDirectPhotonsWorkRange(
 		const FLight* Light = Lights[QuantizedLightIndex];
 
 		FLightRay SampleRay;
-		FVector4 LightSourceNormal;
-		FVector2D LightSurfacePosition;
+		FVector4f LightSourceNormal;
+		FVector2f LightSurfacePosition;
 		float RayDirectionPDF;
 		FLinearColor PathAlpha;
 		{
@@ -551,7 +551,7 @@ void FStaticLightingSystem::EmitDirectPhotonsWorkRange(
 		AggregateMesh->IntersectLightRay(SampleRay, true, true, true, CoherentRayCache, PathIntersection);
 		Output.DirectPhotonsTracingThreadTime += CoherentRayCache.FirstHitRayTraceTime - BeforeDirectTraceTime;
 
-		const FVector4 WorldPathDirection = SampleRay.Direction.GetUnsafeNormal3();
+		const FVector4f WorldPathDirection = SampleRay.Direction.GetUnsafeNormal3();
 
 		// Register this photon path as long as it hit a frontface of something in the scene
 		if (PathIntersection.bIntersects 
@@ -599,7 +599,7 @@ void FStaticLightingSystem::EmitDirectPhotonsWorkRange(
 			if (NumIndirectPathRaysGathered < WorkRange.TargetNumIndirectPhotonPaths)
 			{
 				FStaticLightingVertex IntersectionVertexWithTangents(PathIntersection.IntersectionVertex);
-				FVector4 NewWorldPathDirection;
+				FVector4f NewWorldPathDirection;
 				float BRDFDirectionPDF;
 
 				// Generate a new path direction from the BRDF
@@ -620,10 +620,10 @@ void FStaticLightingSystem::EmitDirectPhotonsWorkRange(
 				const float CosFactor = -Dot3(WorldPathDirection, IntersectionVertexWithTangents.WorldTangentZ);
 				checkSlow(CosFactor >= 0.0f && CosFactor <= 1.0f);
 
-				const FVector4 RayStart = IntersectionVertexWithTangents.WorldPosition 
+				const FVector4f RayStart = IntersectionVertexWithTangents.WorldPosition 
 					+ NewWorldPathDirection * SceneConstants.VisibilityRayOffsetDistance 
 					+ IntersectionVertexWithTangents.WorldTangentZ * SceneConstants.VisibilityNormalOffsetDistance;
-				const FVector4 RayEnd = IntersectionVertexWithTangents.WorldPosition + NewWorldPathDirection * MaxRayDistance;
+				const FVector4f RayEnd = IntersectionVertexWithTangents.WorldPosition + NewWorldPathDirection * MaxRayDistance;
 
 				FLightRay IndirectSampleRay(
 					RayStart,
@@ -689,7 +689,7 @@ void FStaticLightingSystem::BuildPhotonSegmentMap(const FPhotonOctree& SourcePho
 
 /** Gathers indirect photons based on the indirect photon paths. */
 void FStaticLightingSystem::EmitIndirectPhotons(
-	const FBoxSphereBounds& ImportanceBounds,
+	const FBoxSphereBounds3f& ImportanceBounds,
 	const TArray<TArray<FIndirectPathRay> >& IndirectPathRays, 
 	TArray<TArray<FIrradiancePhoton>>& IrradiancePhotons)
 {
@@ -978,8 +978,8 @@ void FStaticLightingSystem::EmitIndirectPhotonsWorkRange(
 			}
 			else
 			{
-				FVector4 LightSourceNormal;
-				FVector2D LightSurfacePosition;
+				FVector4f LightSourceNormal;
+				FVector2f LightSurfacePosition;
 				// No indirect path rays from this light, sample it uniformly
 				Light->SampleDirection(RandomStream, SampleRay, LightSourceNormal, LightSurfacePosition, RayDirectionPDF, PathAlpha);
 			}
@@ -993,7 +993,7 @@ void FStaticLightingSystem::EmitIndirectPhotonsWorkRange(
 			}
 
 			// Clip the end of the photon path to the importance volume, or skip if the photon path does not intersect the importance volume at all.
-			FVector4 ClippedStart, ClippedEnd;
+			FVector4f ClippedStart, ClippedEnd;
 			if (!ClipLineWithBox(Input.ImportanceBounds.GetBox(), SampleRay.Start, SampleRay.End, ClippedStart, ClippedEnd))
 			{
 				continue;
@@ -1009,7 +1009,7 @@ void FStaticLightingSystem::EmitIndirectPhotonsWorkRange(
 		Output.IntersectLightRayThreadTime += CoherentRayCache.FirstHitRayTraceTime - BeforeLightRayTime;
 
 		LIGHTINGSTAT(FScopedRDTSCTimer PhotonTracingTimer(Output.PhotonBounceTracingThreadTime));
-		FVector4 WorldPathDirection = SampleRay.Direction.GetUnsafeNormal3();
+		FVector4f WorldPathDirection = SampleRay.Direction.GetUnsafeNormal3();
 		// Continue tracing this photon path as long as it hit a frontface of something in the scene
 		while (PathIntersection.bIntersects && Dot3(WorldPathDirection, PathIntersection.IntersectionVertex.WorldTangentZ) < 0.0f)
 		{
@@ -1108,7 +1108,7 @@ void FStaticLightingSystem::EmitIndirectPhotonsWorkRange(
 
 			FStaticLightingVertex IntersectionVertexWithTangents(PathIntersection.IntersectionVertex);
 
-			FVector4 NewWorldPathDirection;
+			FVector4f NewWorldPathDirection;
 			float BRDFDirectionPDF;
 
 			// Generate a new path direction from the BRDF
@@ -1161,14 +1161,14 @@ void FStaticLightingSystem::EmitIndirectPhotonsWorkRange(
 
 			checkSlow(FLinearColorUtils::AreFloatsValid(PathAlpha));
 
-			const FVector4 RayStart = IntersectionVertexWithTangents.WorldPosition 
+			const FVector4f RayStart = IntersectionVertexWithTangents.WorldPosition 
 				+ NewWorldPathDirection * SceneConstants.VisibilityRayOffsetDistance 
 				+ IntersectionVertexWithTangents.WorldTangentZ * SceneConstants.VisibilityNormalOffsetDistance;
-			FVector4 RayEnd = IntersectionVertexWithTangents.WorldPosition + NewWorldPathDirection * MaxRayDistance;
+			FVector4f RayEnd = IntersectionVertexWithTangents.WorldPosition + NewWorldPathDirection * MaxRayDistance;
 
 			// Clip photon path end points to the importance volume, so we do not bother tracing rays outside the area that photons can be deposited.
 			// If the photon path does not intersect the importance volume at all, it did not originate from inside the volume, so skip to the next photon.
-			FVector4 ClippedStart, ClippedEnd;
+			FVector4f ClippedStart, ClippedEnd;
 			if (!ClipLineWithBox(Input.ImportanceBounds.GetBox(), RayStart, RayEnd, ClippedStart, ClippedEnd))
 			{
 				break;
@@ -1200,12 +1200,12 @@ void FStaticLightingSystem::EmitIndirectPhotonsWorkRange(
 
 				checkSlow(FLinearColorUtils::AreFloatsValid(PathAlpha));
 
-				FVector4 RayEnd = SampleRay.Start + WorldPathDirection * Input.ImportanceBounds.SphereRadius * 2;
-				FVector4 ClippedStart, ClippedEnd;
+				FVector4f RayEnd = SampleRay.Start + WorldPathDirection * Input.ImportanceBounds.SphereRadius * 2;
+				FVector4f ClippedStart, ClippedEnd;
 				if (ClipLineWithBox(Input.ImportanceBounds.GetBox(), SampleRay.Start, RayEnd, ClippedStart, ClippedEnd))
 				{
 					// Create the description for an escaped photon
-					const FPhoton NewPhoton(Output.NumPhotonsEmitted, ClippedEnd, (ClippedEnd - SampleRay.Start).Size3(), -WorldPathDirection, FVector4(0, 0, 1), PathAlpha);
+					const FPhoton NewPhoton(Output.NumPhotonsEmitted, ClippedEnd, (ClippedEnd - SampleRay.Start).Size3(), -WorldPathDirection, FVector4f(0, 0, 1), PathAlpha);
 					Output.FirstBounceEscapedPhotons.Add(NewPhoton);
 				}
 			}
@@ -1216,7 +1216,7 @@ void FStaticLightingSystem::EmitIndirectPhotonsWorkRange(
 }
 
 /** Iterates through all irradiance photons, searches for nearby direct photons, and marks the irradiance photon has having direct photon influence if necessary. */
-void FStaticLightingSystem::MarkIrradiancePhotons(const FBoxSphereBounds& ImportanceBounds, TArray<TArray<FIrradiancePhoton>>& IrradiancePhotons)
+void FStaticLightingSystem::MarkIrradiancePhotons(const FBoxSphereBounds3f& ImportanceBounds, TArray<TArray<FIrradiancePhoton>>& IrradiancePhotons)
 {
 	check(PhotonMappingSettings.bUseIrradiancePhotons);
 	GSwarm->SendMessage( NSwarm::FTimingMessage( NSwarm::PROGSTATE_Preparing2, 0 ) );
@@ -1347,7 +1347,7 @@ void FStaticLightingSystem::MarkIrradiancePhotonsWorkRange(
 }
 
 /** Calculates irradiance for photons randomly chosen to precalculate irradiance. */
-void FStaticLightingSystem::CalculateIrradiancePhotons(const FBoxSphereBounds& ImportanceBounds, TArray<TArray<FIrradiancePhoton>>& IrradiancePhotons)
+void FStaticLightingSystem::CalculateIrradiancePhotons(const FBoxSphereBounds3f& ImportanceBounds, TArray<TArray<FIrradiancePhoton>>& IrradiancePhotons)
 {
 	check(PhotonMappingSettings.bUseIrradiancePhotons);
 	//@todo - add a preparing stage for the swarm visualizer
@@ -1665,13 +1665,13 @@ void FStaticLightingSystem::CacheIrradiancePhotonsThreadLoop(int32 ThreadIndex,
 /** Returns true if a photon was found within MaxPhotonSearchDistance. */
 bool FStaticLightingSystem::FindAnyNearbyPhoton(
 	const FPhotonOctree& PhotonMap, 
-	const FVector4& SearchPosition, 
+	const FVector4f& SearchPosition, 
 	float MaxPhotonSearchDistance,
 	bool bDebugThisLookup) const
 {
 	FPlatformAtomics::InterlockedIncrement(&Stats.NumPhotonGathers);
 
-	const FBox SearchBox = FBox::BuildAABB(SearchPosition, FVector4(MaxPhotonSearchDistance, MaxPhotonSearchDistance, MaxPhotonSearchDistance));
+	const FBox3f SearchBox = FBox3f::BuildAABB(SearchPosition, FVector4f(MaxPhotonSearchDistance, MaxPhotonSearchDistance, MaxPhotonSearchDistance));
 	for (FPhotonOctree::TConstIterator<> OctreeIt(PhotonMap); OctreeIt.HasPendingNodes(); OctreeIt.Advance())
 	{
 		const FPhotonOctree::FNode& CurrentNode = OctreeIt.GetCurrentNode();
@@ -1726,8 +1726,8 @@ bool FStaticLightingSystem::FindAnyNearbyPhoton(
  */
 float FStaticLightingSystem::FindNearbyPhotonsIterative(
 	const FPhotonOctree& PhotonMap, 
-	const FVector4& SearchPosition, 
-	const FVector4& SearchNormal, 
+	const FVector4f& SearchPosition, 
+	const FVector4f& SearchNormal, 
 	int32 NumPhotonsToFind,
 	float StartPhotonSearchDistance, 
 	float MaxPhotonSearchDistance,
@@ -1766,7 +1766,7 @@ float FStaticLightingSystem::FindNearbyPhotonsIterative(
 		// Presize to avoid unnecessary allocations
 		// Empty last search iteration's results so we don't have to use AddUniqueItem
 		FoundPhotons.Empty(FMath::Max(NumPhotonsToFind, FoundPhotons.Num() + FoundPhotons.GetSlack()));
-		const FBox SearchBox = FBox::BuildAABB(SearchPosition, FVector4(SearchDistance, SearchDistance, SearchDistance));
+		const FBox3f SearchBox = FBox3f::BuildAABB(SearchPosition, FVector4f(SearchDistance, SearchDistance, SearchDistance));
 		for (FPhotonOctree::TConstIterator<TInlineAllocator<600>> OctreeIt(PhotonMap); OctreeIt.HasPendingNodes(); OctreeIt.Advance())
 		{
 			const FPhotonOctree::FNode& CurrentNode = OctreeIt.GetCurrentNode();
@@ -1890,7 +1890,7 @@ float FStaticLightingSystem::FindNearbyPhotonsIterative(
  */
 float FStaticLightingSystem::FindNearbyPhotonsInVolumeIterative(
 	const FPhotonSegmentOctree& PhotonSegmentMap, 
-	const FVector4& SearchPosition, 
+	const FVector4f& SearchPosition, 
 	int32 NumPhotonsToFind,
 	float StartPhotonSearchDistance, 
 	float MaxPhotonSearchDistance,
@@ -1908,7 +1908,7 @@ float FStaticLightingSystem::FindNearbyPhotonsInVolumeIterative(
 		// Presize to avoid unnecessary allocations
 		// Empty last search iteration's results so we don't have to use AddUniqueItem
 		FoundPhotonSegments.Empty(FMath::Max(NumPhotonsToFind, FoundPhotonSegments.Num() + FoundPhotonSegments.GetSlack()));
-		const FBox SearchBox = FBox::BuildAABB(SearchPosition, FVector4(SearchDistance, SearchDistance, SearchDistance));
+		const FBox3f SearchBox = FBox3f::BuildAABB(SearchPosition, FVector4f(SearchDistance, SearchDistance, SearchDistance));
 		for (FPhotonSegmentOctree::TConstIterator<TInlineAllocator<600>> OctreeIt(PhotonSegmentMap); OctreeIt.HasPendingNodes(); OctreeIt.Advance())
 		{
 			const FPhotonSegmentOctree::FNode& CurrentNode = OctreeIt.GetCurrentNode();
@@ -2021,8 +2021,8 @@ struct FOctreeNodeRefAndDistance
  */
 float FStaticLightingSystem::FindNearbyPhotonsSorted(
 	const FPhotonOctree& PhotonMap, 
-	const FVector4& SearchPosition, 
-	const FVector4& SearchNormal, 
+	const FVector4f& SearchPosition, 
+	const FVector4f& SearchNormal, 
 	int32 NumPhotonsToFind, 
 	float MaxPhotonSearchDistance,
 	bool bDebugSearchResults,
@@ -2197,7 +2197,7 @@ float FStaticLightingSystem::FindNearbyPhotonsSorted(
 class FCompareNearestIrradiancePhotons
 {
 public:
-	FCompareNearestIrradiancePhotons(const FVector4& InPosition) :
+	FCompareNearestIrradiancePhotons(const FVector4f& InPosition) :
 		ComparePosition(InPosition)
 	{}
 
@@ -2209,7 +2209,7 @@ public:
 	}
 
 private:
-	const FVector4 ComparePosition;
+	const FVector4f ComparePosition;
 };
 
 /** Finds the nearest irradiance photon, if one exists. */
@@ -2229,7 +2229,7 @@ FIrradiancePhoton* FStaticLightingSystem::FindNearestIrradiancePhoton(
 
 	// Empty the temporary array without reallocating
 	TempIrradiancePhotons.Empty(TempIrradiancePhotons.Num() + TempIrradiancePhotons.GetSlack());
-	const FBox SearchBox = FBox::BuildAABB(Vertex.WorldPosition, FVector4(SearchDistance, SearchDistance, SearchDistance));
+	const FBox3f SearchBox = FBox3f::BuildAABB(Vertex.WorldPosition, FVector4f(SearchDistance, SearchDistance, SearchDistance));
 	{
 		LIGHTINGSTAT(FScopedRDTSCTimer OctreeTraversal(MappingContext.Stats.IrradiancePhotonOctreeTraversalTime));
 		for (FIrradiancePhotonOctree::TConstIterator<> OctreeIt(IrradiancePhotonMap); OctreeIt.HasPendingNodes(); OctreeIt.Advance())
@@ -2258,7 +2258,7 @@ FIrradiancePhoton* FStaticLightingSystem::FindNearestIrradiancePhoton(
 			{
 				FIrradiancePhotonElement& PhotonElement = *MeshIt;
 				FIrradiancePhoton& CurrentPhoton = PhotonElement.GetPhoton();
-				const FVector4 PhotonToVertexVector = Vertex.WorldPosition - CurrentPhoton.GetPosition();
+				const FVector4f PhotonToVertexVector = Vertex.WorldPosition - CurrentPhoton.GetPosition();
 				const float DistanceSquared = PhotonToVertexVector.SizeSquared3();
 				const float CosTheta = Dot3(Vertex.WorldTangentZ, CurrentPhoton.GetSurfaceNormal());
 
@@ -2300,7 +2300,7 @@ FIrradiancePhoton* FStaticLightingSystem::FindNearestIrradiancePhoton(
 		for (int32 PhotonIndex = 0; PhotonIndex < TempIrradiancePhotons.Num(); PhotonIndex++)
 		{
 			FIrradiancePhoton* CurrentPhoton = TempIrradiancePhotons[PhotonIndex];
-			const FVector4 VertexToPhoton = CurrentPhoton->GetPosition() - Vertex.WorldPosition;
+			const FVector4f VertexToPhoton = CurrentPhoton->GetPosition() - Vertex.WorldPosition;
 			const FLightRay VertexToPhotonRay(
 				Vertex.WorldPosition + VertexToPhoton.GetSafeNormal() * SceneConstants.VisibilityRayOffsetDistance + Vertex.WorldTangentZ * SceneConstants.VisibilityNormalOffsetDistance,
 				CurrentPhoton->GetPosition() + CurrentPhoton->GetSurfaceNormal() * SceneConstants.VisibilityNormalOffsetDistance,
@@ -2444,7 +2444,7 @@ FGatheredLightSample FStaticLightingSystem::CalculatePhotonIncidentRadiance(
 		for (int32 PhotonIndex = 0; PhotonIndex < FoundPhotons.Num(); PhotonIndex++)
 		{
 			const FPhoton& CurrentPhoton = FoundPhotons[PhotonIndex];
-			const FVector4 TangentPathDirection = Vertex.TransformWorldVectorToTangent(CurrentPhoton.GetIncidentDirection());
+			const FVector4f TangentPathDirection = Vertex.TransformWorldVectorToTangent(CurrentPhoton.GetIncidentDirection());
 			if (TangentPathDirection.Z > 0)
 			{
 				const float PhotonDistance = (CurrentPhoton.GetPosition() - Vertex.WorldPosition).Size3();
@@ -2465,7 +2465,7 @@ FLinearColor FStaticLightingSystem::CalculatePhotonExitantRadiance(
 	const FStaticLightingMesh* Mesh,
 	const FMinimalStaticLightingVertex& Vertex,
 	int32 ElementIndex,
-	const FVector4& OutgoingDirection,
+	const FVector4f& OutgoingDirection,
 	bool bDebugThisDensityEstimation) const
 {
 	TArray<FPhoton> FoundPhotons;
