Fix CPU Lightmass by coverting LWC types to single precision

diff --git a/Engine/Source/Programs/UnrealLightmass/Private/Lighting/PrecomputedVisibility.cpp b/Engine/Source/Programs/UnrealLightmass/Private/Lighting/PrecomputedVisibility.cpp
--- a/Engine/Source/Programs/UnrealLightmass/Private/Lighting/PrecomputedVisibility.cpp
+++ b/Engine/Source/Programs/UnrealLightmass/Private/Lighting/PrecomputedVisibility.cpp
@@ -15,7 +15,7 @@ struct FVisibilitySamplePos
 	 * Min and max world space heights of a single sample on a triangle. 
 	 * This is necessary because supersampling is used during rasterization.
 	 */
-	FVector2D HeightRange;
+	FVector2f HeightRange;
 };
 
 struct FCellHeights
@@ -23,7 +23,7 @@ struct FCellHeights
 	/** Last triangle index that rasterized to this cell. */
 	uint64 TriangleIndex;
 	/** World space X and Y position of this cell. */
-	FVector2D Position;
+	FVector2f Position;
 	/** Array of triangle hits on this cell. */
 	TArray<FVisibilitySamplePos> HitTriangles;
 };
@@ -82,14 +82,14 @@ class FCellPlacementRasterPolicy
 {
 public:
 
-	typedef FVector4 InterpolantType;
+	typedef FVector4f InterpolantType;
 
 	/** Initialization constructor. */
 	FCellPlacementRasterPolicy(
 		FCellToHeightsMap& InHeightsMap,
 		const FScene& InScene, 
 		FStaticLightingSystem& InSystem,
-		const FBoxSphereBounds& InPrecomputedVisibilityBounds,
+		const FBoxSphereBounds3f& InPrecomputedVisibilityBounds,
 		float InCellSize)
 		:
 		HeightsMap(InHeightsMap),
@@ -122,7 +122,7 @@ private:
 	FCellToHeightsMap& HeightsMap;
 	const FScene& Scene;
 	FStaticLightingSystem& System;
-	FBoxSphereBounds PrecomputedVisibilityBounds;
+	FBoxSphereBounds3f PrecomputedVisibilityBounds;
 	float CellSize;
 };
 
@@ -136,10 +136,10 @@ void FCellPlacementRasterPolicy::ProcessPixel(int32 X,int32 Y,const InterpolantT
 		{
 			// If this is the first hit on this cell from the current triangle, add a new sample
 			FVisibilitySamplePos Sample;
-			const FVector GridPosition = PrecomputedVisibilityBounds.Origin - PrecomputedVisibilityBounds.BoxExtent + FVector(X, Y, 0) * CellSize;
-			Sample.HeightRange = FVector2D(WorldPosition.Z, WorldPosition.Z);
+			const FVector3f GridPosition = PrecomputedVisibilityBounds.Origin - PrecomputedVisibilityBounds.BoxExtent + FVector3f(X, Y, 0) * CellSize;
+			Sample.HeightRange = FVector2f(WorldPosition.Z, WorldPosition.Z);
 			Cell.HitTriangles.Add(Sample);
-			Cell.Position = FVector2D(GridPosition.X, GridPosition.Y);
+			Cell.Position = FVector2f(GridPosition.X, GridPosition.Y);
 			Cell.TriangleIndex = TriangleIndex;
 		}
 		else
@@ -152,9 +152,9 @@ void FCellPlacementRasterPolicy::ProcessPixel(int32 X,int32 Y,const InterpolantT
 	}
 }
 
-int32 FStaticLightingSystem::GetGroupCellIndex(FVector BoxCenter) const
+int32 FStaticLightingSystem::GetGroupCellIndex(FVector3f BoxCenter) const
 {
-	const FVector GridPosition = FVector(GroupVisibilityGridSizeXY, GroupVisibilityGridSizeXY, GroupVisibilityGridSizeZ) * (BoxCenter - VisibilityGridBounds.Min) / (VisibilityGridBounds.Max - VisibilityGridBounds.Min);
+	const FVector3f GridPosition = FVector3f(GroupVisibilityGridSizeXY, GroupVisibilityGridSizeXY, GroupVisibilityGridSizeZ) * (BoxCenter - VisibilityGridBounds.Min) / (VisibilityGridBounds.Max - VisibilityGridBounds.Min);
 	const int32 GridX = FMath::TruncToInt(GridPosition.X);
 	const int32 GridY = FMath::TruncToInt(GridPosition.Y);
 	const int32 GridZ = FMath::TruncToInt(GridPosition.Z);
@@ -177,10 +177,10 @@ class FVisibilityMeshSortInfo
 public:
 	float Distance;
 	int32 Index;
-	FBox Bounds;
+	FBox3f Bounds;
 
 	FVisibilityMeshSortInfo() :
-		Bounds(FBox(ForceInit))
+		Bounds(FBox3f(ForceInit))
 	{}
 };
 
@@ -190,8 +190,8 @@ void FStaticLightingSystem::SetupPrecomputedVisibility()
 	const double StartTime = FPlatformTime::Seconds();
 	FLMRandomStream RandomStream(0);
 
-	const FBoxSphereBounds PrecomputedVisibilityBounds = Scene.GetVisibilityVolumeBounds();
-	const FVector4 VolumeSizes = PrecomputedVisibilityBounds.BoxExtent * 2.0f / PrecomputedVisibilitySettings.CellSize;
+	const FBoxSphereBounds3f PrecomputedVisibilityBounds = Scene.GetVisibilityVolumeBounds();
+	const FVector4f VolumeSizes = PrecomputedVisibilityBounds.BoxExtent * 2.0f / PrecomputedVisibilitySettings.CellSize;
 	const int32 SizeX = FMath::TruncToInt(VolumeSizes.X + DELTA) + 1;
 	const int32 SizeY = FMath::TruncToInt(VolumeSizes.Y + DELTA) + 1;
 
@@ -246,38 +246,38 @@ void FStaticLightingSystem::SetupPrecomputedVisibility()
 					if (!PrecomputedVisibilitySettings.bPlaceCellsOnOpaqueOnly || bOpaqueMesh
 						|| (!CurrentMesh->IsMasked(ElementIndex) && !CurrentMesh->IsTranslucent(ElementIndex)))
 					{
-						FVector2D XYPositions[3];
+						FVector2f XYPositions[3];
 						for (int32 VertIndex = 0; VertIndex < 3; VertIndex++)
 						{
 							// Transform world space positions from [PrecomputedVisibilityBounds.Origin - PrecomputedVisibilityBounds.BoxExtent, PrecomputedVisibilityBounds.Origin + PrecomputedVisibilityBounds.BoxExtent] into [0,1]
-							const FVector4 TransformedPosition = (Vertices[VertIndex].WorldPosition - PrecomputedVisibilityBounds.Origin + PrecomputedVisibilityBounds.BoxExtent) / (2.0f * PrecomputedVisibilityBounds.BoxExtent);
+							const FVector4f TransformedPosition = (Vertices[VertIndex].WorldPosition - PrecomputedVisibilityBounds.Origin + PrecomputedVisibilityBounds.BoxExtent) / (2.0f * PrecomputedVisibilityBounds.BoxExtent);
 							// Project positions onto the XY plane
-							XYPositions[VertIndex] = FVector2D(TransformedPosition.X * SizeX, TransformedPosition.Y * SizeY);
+							XYPositions[VertIndex] = FVector2f(TransformedPosition.X * SizeX, TransformedPosition.Y * SizeY);
 						}
 
-						const FVector4 TriangleNormal = (Vertices[2].WorldPosition - Vertices[0].WorldPosition) ^ (Vertices[1].WorldPosition - Vertices[0].WorldPosition);
+						const FVector4f TriangleNormal = (Vertices[2].WorldPosition - Vertices[0].WorldPosition) ^ (Vertices[1].WorldPosition - Vertices[0].WorldPosition);
 
 						// Only rasterize upward facing triangles
 						if (TriangleNormal.Z > 0.0f)
 						{
 							Rasterizer.SetTriangleIndex(NextTriangleIndex);
 
-							FVector2D SubsamplePositions[9];
-							SubsamplePositions[0] = FVector2D(.5f, .5f);
-							SubsamplePositions[1] = FVector2D(0, .5f);
-							SubsamplePositions[2] = FVector2D(.5f, 0);
-							SubsamplePositions[3] = FVector2D(1, .5f);
-							SubsamplePositions[4] = FVector2D(.5f, 1);
-							SubsamplePositions[5] = FVector2D(1, 1);
-							SubsamplePositions[6] = FVector2D(0, 1);
-							SubsamplePositions[7] = FVector2D(1, 0);
-							SubsamplePositions[8] = FVector2D(0, 0);
+							FVector2f SubsamplePositions[9];
+							SubsamplePositions[0] = FVector2f(.5f, .5f);
+							SubsamplePositions[1] = FVector2f(0, .5f);
+							SubsamplePositions[2] = FVector2f(.5f, 0);
+							SubsamplePositions[3] = FVector2f(1, .5f);
+							SubsamplePositions[4] = FVector2f(.5f, 1);
+							SubsamplePositions[5] = FVector2f(1, 1);
+							SubsamplePositions[6] = FVector2f(0, 1);
+							SubsamplePositions[7] = FVector2f(1, 0);
+							SubsamplePositions[8] = FVector2f(0, 0);
 
 							const float EdgePullback = .1f;
 
 							for (int32 SampleIndex = 0; SampleIndex < UE_ARRAY_COUNT(SubsamplePositions); SampleIndex++)
 							{
-								const FVector2D SamplePosition = SubsamplePositions[SampleIndex] * (1 - 2 * EdgePullback) + FVector2D(EdgePullback, EdgePullback);
+								const FVector2f SamplePosition = SubsamplePositions[SampleIndex] * (1 - 2 * EdgePullback) + FVector2f(EdgePullback, EdgePullback);
 
 								Rasterizer.DrawTriangle(
 									Vertices[0].WorldPosition,
@@ -299,13 +299,13 @@ void FStaticLightingSystem::SetupPrecomputedVisibility()
 
 		AllPrecomputedVisibilityCells.Empty(SizeX * SizeY * 2);
 
-		TArray<FVector2D> PlacedHeightRanges;
+		TArray<FVector2f> PlacedHeightRanges;
 		for (int32 Y = 0; Y < SizeY; Y++)
 		{
 			for (int32 X = 0; X < SizeX; X++)
 			{
 				FCellHeights& Cell = HeightsMap(X, Y);
-				const FVector2D CurrentPosition = Cell.Position;
+				const FVector2f CurrentPosition = Cell.Position;
 
 				// Sort the heights from smallest to largest
 				struct FCompareSampleZ
@@ -334,19 +334,19 @@ void FStaticLightingSystem::SetupPrecomputedVisibility()
 						&& CurrentMaxHeight - LastSampleHeight > PrecomputedVisibilitySettings.PlayAreaHeight))
 					{
 						FPrecomputedVisibilityCell NewCell;
-						NewCell.Bounds = FBox(
-							FVector4(
+						NewCell.Bounds = FBox3f(
+							FVector4f(
 								CurrentPosition.X - PrecomputedVisibilitySettings.CellSize / 2,
 								CurrentPosition.Y - PrecomputedVisibilitySettings.CellSize / 2,
 								CurrentMaxHeight),
-							FVector4(
+							FVector4f(
 								CurrentPosition.X + PrecomputedVisibilitySettings.CellSize / 2,
 								CurrentPosition.Y + PrecomputedVisibilitySettings.CellSize / 2,
 								CurrentMaxHeight + PrecomputedVisibilitySettings.PlayAreaHeight));
 
 						AllPrecomputedVisibilityCells.Add(NewCell);
 						LastSampleHeight = CurrentMaxHeight;
-						PlacedHeightRanges.Add(FVector2D(NewCell.Bounds.Min.Z, NewCell.Bounds.Max.Z));
+						PlacedHeightRanges.Add(FVector2f(NewCell.Bounds.Min.Z, NewCell.Bounds.Max.Z));
 					}
 				}
 
@@ -371,7 +371,7 @@ void FStaticLightingSystem::SetupPrecomputedVisibility()
 
 							for (int32 PlacedHeightIndex = 0; PlacedHeightIndex < PlacedHeightRanges.Num(); PlacedHeightIndex++)
 							{
-								FVector2D CellHeightRange = PlacedHeightRanges[PlacedHeightIndex];
+								FVector2f CellHeightRange = PlacedHeightRanges[PlacedHeightIndex];
 
 								if (TestHeight > CellHeightRange.X && TestHeight < CellHeightRange.Y)
 								{
@@ -396,7 +396,7 @@ void FStaticLightingSystem::SetupPrecomputedVisibility()
 
 								if (ClosestCellInZIndex >= 0)
 								{
-									const FVector2D NearestCellHeightRange = PlacedHeightRanges[ClosestCellInZIndex];
+									const FVector2f NearestCellHeightRange = PlacedHeightRanges[ClosestCellInZIndex];
 									const float NearestCellCompareHeight = (NearestCellHeightRange.X + NearestCellHeightRange.Y) / 2;
 
 									// Move the bottom of the cell to be placed such that it doesn't overlap the nearest cell
@@ -411,18 +411,18 @@ void FStaticLightingSystem::SetupPrecomputedVisibility()
 									}
 								}
 
-								NewCell.Bounds = FBox(
-									FVector4(
+								NewCell.Bounds = FBox3f(
+									FVector4f(
 									CurrentPosition.X - PrecomputedVisibilitySettings.CellSize / 2,
 									CurrentPosition.Y - PrecomputedVisibilitySettings.CellSize / 2,
 									DesiredCellBottom),
-									FVector4(
+									FVector4f(
 									CurrentPosition.X + PrecomputedVisibilitySettings.CellSize / 2,
 									CurrentPosition.Y + PrecomputedVisibilitySettings.CellSize / 2,
 									DesiredCellBottom + PrecomputedVisibilitySettings.PlayAreaHeight));
 
 								AllPrecomputedVisibilityCells.Add(NewCell);
-								PlacedHeightRanges.Add(FVector2D(NewCell.Bounds.Min.Z, NewCell.Bounds.Max.Z));
+								PlacedHeightRanges.Add(FVector2f(NewCell.Bounds.Min.Z, NewCell.Bounds.Max.Z));
 							}
 						}
 					}
@@ -435,7 +435,7 @@ void FStaticLightingSystem::SetupPrecomputedVisibility()
 	const int32 NumCellsPlacedOnSurfaces = AllPrecomputedVisibilityCells.Num();
 	for (int32 CameraPositionIndex = 0; CameraPositionIndex < Scene.CameraTrackPositions.Num(); CameraPositionIndex++)
 	{
-		const FVector4& CurrentPosition = Scene.CameraTrackPositions[CameraPositionIndex];
+		const FVector4f& CurrentPosition = Scene.CameraTrackPositions[CameraPositionIndex];
 		bool bInsideCell = false;
 		for (int32 CellIndex = 0; CellIndex < AllPrecomputedVisibilityCells.Num(); CellIndex++)
 		{
@@ -452,8 +452,8 @@ void FStaticLightingSystem::SetupPrecomputedVisibility()
 
 			// Snap the cell min to the nearest factor of CellSize from the visibility bounds min + CellSize / 2
 			// The CellSize / 2 offset is necessary to match up with cells produced by the rasterizer, since pixels are rasterized at cell centers
-			const FVector4 PreSnapTranslation = FVector(PrecomputedVisibilitySettings.CellSize / 2) + PrecomputedVisibilityBounds.Origin - PrecomputedVisibilityBounds.BoxExtent;
-			const FVector4 TranslatedPosition = CurrentPosition - PreSnapTranslation;
+			const FVector4f PreSnapTranslation = FVector3f(PrecomputedVisibilitySettings.CellSize / 2) + PrecomputedVisibilityBounds.Origin - PrecomputedVisibilityBounds.BoxExtent;
+			const FVector4f TranslatedPosition = CurrentPosition - PreSnapTranslation;
 			// FMath::Fmod gives the offset to round up for negative numbers, when we always want the offset to round down
 			const float XOffset = TranslatedPosition.X > 0.0f ? 
 				FMath::Fmod(TranslatedPosition.X, PrecomputedVisibilitySettings.CellSize) : 
@@ -461,14 +461,14 @@ void FStaticLightingSystem::SetupPrecomputedVisibility()
 			const float YOffset = TranslatedPosition.Y > 0.0f ? 
 				FMath::Fmod(TranslatedPosition.Y, PrecomputedVisibilitySettings.CellSize) : 
 				PrecomputedVisibilitySettings.CellSize - FMath::Fmod(-TranslatedPosition.Y, PrecomputedVisibilitySettings.CellSize);
-			const FVector4 SnappedPosition(CurrentPosition.X - XOffset, CurrentPosition.Y - YOffset, CurrentPosition.Z);
+			const FVector4f SnappedPosition(CurrentPosition.X - XOffset, CurrentPosition.Y - YOffset, CurrentPosition.Z);
 
-			NewCell.Bounds = FBox(
-				FVector4(
+			NewCell.Bounds = FBox3f(
+				FVector4f(
 					SnappedPosition.X,
 					SnappedPosition.Y,
 					SnappedPosition.Z - .5f * PrecomputedVisibilitySettings.PlayAreaHeight),
-				FVector4(
+				FVector4f(
 					SnappedPosition.X + PrecomputedVisibilitySettings.CellSize,
 					SnappedPosition.Y + PrecomputedVisibilitySettings.CellSize,
 					SnappedPosition.Z + .5f * PrecomputedVisibilitySettings.PlayAreaHeight));
@@ -484,7 +484,7 @@ void FStaticLightingSystem::SetupPrecomputedVisibility()
 		TArray<FVisibilityMeshSortInfo> SortMeshes;
 		SortMeshes.Empty(VisibilityMeshes.Num());
 
-		FVector CenterPosition(0, 0, 0);
+		FVector3f CenterPosition(0, 0, 0);
 
 		if (VisibilityMeshes.Num() > 0)
 		{
@@ -511,20 +511,20 @@ void FStaticLightingSystem::SetupPrecomputedVisibility()
 			SortMeshes.Add(NewInfo);
 		}
 
-		FVector CubeCorners[8];
-		CubeCorners[0] = FVector(1, 1, 1);
-		CubeCorners[1] = FVector(-1, 1, 1);
-		CubeCorners[2] = FVector(1, -1, 1);
-		CubeCorners[3] = FVector(-1, -1, 1);
-		CubeCorners[4] = FVector(1, 1, -1);
-		CubeCorners[5] = FVector(-1, 1, -1);
-		CubeCorners[6] = FVector(1, -1, -1);
-		CubeCorners[7] = FVector(-1, -1, -1);
+		FVector3f CubeCorners[8];
+		CubeCorners[0] = FVector3f(1, 1, 1);
+		CubeCorners[1] = FVector3f(-1, 1, 1);
+		CubeCorners[2] = FVector3f(1, -1, 1);
+		CubeCorners[3] = FVector3f(-1, -1, 1);
+		CubeCorners[4] = FVector3f(1, 1, -1);
+		CubeCorners[5] = FVector3f(-1, 1, -1);
+		CubeCorners[6] = FVector3f(1, -1, -1);
+		CubeCorners[7] = FVector3f(-1, -1, -1);
 
 		for (int32 MeshIndex = 0; MeshIndex < SortMeshes.Num(); MeshIndex++)
 		{
-			const FVector BoxCenter = SortMeshes[MeshIndex].Bounds.GetCenter();
-			const FVector BoxExtent = SortMeshes[MeshIndex].Bounds.GetExtent();
+			const FVector3f BoxCenter = SortMeshes[MeshIndex].Bounds.GetCenter();
+			const FVector3f BoxExtent = SortMeshes[MeshIndex].Bounds.GetExtent();
 			float LocalDistance = 0;
 
 			for (int32 CornerIndex = 0; CornerIndex < UE_ARRAY_COUNT(CubeCorners); CornerIndex++)
@@ -545,7 +545,7 @@ void FStaticLightingSystem::SetupPrecomputedVisibility()
 		};
 		SortMeshes.Sort(FCompareMeshesByDistance());
 
-		VisibilityGridBounds = FBox(ForceInit);
+		VisibilityGridBounds = FBox3f(ForceInit);
 		// Drop last 10% of meshes which will expand the grid bounds 
 		// This is to handle distant skybox type meshes
 		const int32 MaxMeshIndex = FMath::Min(FMath::Max(FMath::TruncToInt(.9f * SortMeshes.Num()), 1), SortMeshes.Num() - 1);
@@ -565,7 +565,7 @@ void FStaticLightingSystem::SetupPrecomputedVisibility()
 		const int32 GridSizeXY = GroupVisibilityGridSizeXY;
 		const int32 GridSizeZ = GroupVisibilityGridSizeZ;
 
-		const FVector CellSize = VisibilityGridBounds.GetExtent() / FVector(GroupVisibilityGridSizeXY, GroupVisibilityGridSizeXY, GroupVisibilityGridSizeZ);
+		const FVector3f CellSize = VisibilityGridBounds.GetExtent() / FVector3f(GroupVisibilityGridSizeXY, GroupVisibilityGridSizeXY, GroupVisibilityGridSizeZ);
 		const float GridCellBoundingRadius = CellSize.Size();
 		const float MeshGroupingCellRadiusThreshold = .5f;
 
@@ -579,7 +579,7 @@ void FStaticLightingSystem::SetupPrecomputedVisibility()
 
 		for (int32 VisibilityMeshIndex = 0; VisibilityMeshIndex < VisibilityMeshes.Num(); VisibilityMeshIndex++)
 		{
-			FBox MeshBounds(ForceInit);
+			FBox3f MeshBounds(ForceInit);
 
 			for (int32 OriginalMeshIndex = 0; OriginalMeshIndex < VisibilityMeshes[VisibilityMeshIndex].Meshes.Num(); OriginalMeshIndex++)
 			{
@@ -630,7 +630,7 @@ void FStaticLightingSystem::SetupPrecomputedVisibility()
 		for (int32 GroupIndex = 0; GroupIndex < VisibilityGroups.Num(); GroupIndex++)
 		{
 			FVisibilityMeshGroup& Group = VisibilityGroups[GroupIndex];
-			Group.GroupBounds = FBox(ForceInit);
+			Group.GroupBounds = FBox3f(ForceInit);
 
 			for (int32 EntryIndex = 0; EntryIndex < Group.VisibilityIds.Num(); EntryIndex++)
 			{
@@ -673,28 +673,28 @@ public:
 
 	FAxisAlignedCellFace() {}
 
-	FAxisAlignedCellFace(const FVector4& InFaceDirection, const FVector4& InFaceMin, const FVector4& InFaceExtent) :
+	FAxisAlignedCellFace(const FVector4f& InFaceDirection, const FVector4f& InFaceMin, const FVector4f& InFaceExtent) :
 		FaceDirection(InFaceDirection),
 		FaceMin(InFaceMin),
 		FaceExtent(InFaceExtent)
 	{}
 
-	FVector4 FaceDirection;
-	FVector4 FaceMin;
-	FVector4 FaceExtent;
+	FVector4f FaceDirection;
+	FVector4f FaceMin;
+	FVector4f FaceExtent;
 };
 
 /** Stores information about a single query sample between a visibility cell and a mesh. */
 struct FVisibilityQuerySample
 {
 	/** Sample position generated from the mesh. */
-	FVector4 MeshPosition;
+	FVector4f MeshPosition;
 
 	/** Sample position generated from the cell. */
-	FVector4 CellPosition;
+	FVector4f CellPosition;
 
 	/** Position of the intersection with the scene between the sample positions. */
-	FVector4 IntersectionPosition;
+	FVector4f IntersectionPosition;
 
 	/** Distance along the vector perpendicular to the mesh->cell vector. */
 	float PerpendicularDistance;
@@ -709,7 +709,7 @@ bool ComputeBoxVisibility(
 	FPrecomputedVisibilitySettings& PrecomputedVisibilitySettings,
 	FPrecomputedVisibilityCell& CurrentCell, 
 	FAxisAlignedCellFace* CellFaces,
-	const FBox& MeshBox,
+	const FBox3f& MeshBox,
 	FStaticLightingMappingContext& MappingContext,
 	FLMRandomStream& RandomStream,
 	TArray<int32>& VisibleCellFaces,
@@ -725,18 +725,18 @@ bool ComputeBoxVisibility(
 {
 	const double SampleGenerationStartTime = FPlatformTime::Seconds();
 
-	const FVector4 CenterCellPosition = .5f * (CurrentCell.Bounds.Min + CurrentCell.Bounds.Max);
-	const FVector4 MeshToCellCenter = CenterCellPosition - MeshBox.GetCenter();
+	const FVector4f CenterCellPosition = .5f * (CurrentCell.Bounds.Min + CurrentCell.Bounds.Max);
+	const FVector4f MeshToCellCenter = CenterCellPosition - MeshBox.GetCenter();
 	const float Distance = MeshToCellCenter.Size3();
-	const FVector4 MeshBoxExtent = MeshBox.GetExtent() * 2;
+	const FVector4f MeshBoxExtent = MeshBox.GetExtent() * 2;
 
 	FAxisAlignedCellFace MeshBoxFaces[6];
-	MeshBoxFaces[0] = FAxisAlignedCellFace(FVector4(-1, 0, 0), FVector4(MeshBox.Min.X, MeshBox.Min.Y, MeshBox.Min.Z), FVector4(0, MeshBoxExtent.Y, MeshBoxExtent.Z));
-	MeshBoxFaces[1] = FAxisAlignedCellFace(FVector4(1, 0, 0), FVector4(MeshBox.Min.X + MeshBoxExtent.X, MeshBox.Min.Y, MeshBox.Min.Z), FVector4(0, MeshBoxExtent.Y, MeshBoxExtent.Z));
-	MeshBoxFaces[2] = FAxisAlignedCellFace(FVector4(0, -1, 0), FVector4(MeshBox.Min.X, MeshBox.Min.Y, MeshBox.Min.Z), FVector4(MeshBoxExtent.X, 0, MeshBoxExtent.Z));
-	MeshBoxFaces[3] = FAxisAlignedCellFace(FVector4(0, 1, 0), FVector4(MeshBox.Min.X, MeshBox.Min.Y + MeshBoxExtent.Y, MeshBox.Min.Z), FVector4(MeshBoxExtent.X, 0, MeshBoxExtent.Z));
-	MeshBoxFaces[4] = FAxisAlignedCellFace(FVector4(0, 0, -1), FVector4(MeshBox.Min.X, MeshBox.Min.Y, MeshBox.Min.Z), FVector4(MeshBoxExtent.X, MeshBoxExtent.Y, 0));
-	MeshBoxFaces[5] = FAxisAlignedCellFace(FVector4(0, 0, 1), FVector4(MeshBox.Min.X, MeshBox.Min.Y, MeshBox.Min.Z + MeshBoxExtent.Z), FVector4(MeshBoxExtent.X, MeshBoxExtent.Y, 0));
+	MeshBoxFaces[0] = FAxisAlignedCellFace(FVector4f(-1, 0, 0), FVector4f(MeshBox.Min.X, MeshBox.Min.Y, MeshBox.Min.Z), FVector4f(0, MeshBoxExtent.Y, MeshBoxExtent.Z));
+	MeshBoxFaces[1] = FAxisAlignedCellFace(FVector4f(1, 0, 0), FVector4f(MeshBox.Min.X + MeshBoxExtent.X, MeshBox.Min.Y, MeshBox.Min.Z), FVector4f(0, MeshBoxExtent.Y, MeshBoxExtent.Z));
+	MeshBoxFaces[2] = FAxisAlignedCellFace(FVector4f(0, -1, 0), FVector4f(MeshBox.Min.X, MeshBox.Min.Y, MeshBox.Min.Z), FVector4f(MeshBoxExtent.X, 0, MeshBoxExtent.Z));
+	MeshBoxFaces[3] = FAxisAlignedCellFace(FVector4f(0, 1, 0), FVector4f(MeshBox.Min.X, MeshBox.Min.Y + MeshBoxExtent.Y, MeshBox.Min.Z), FVector4f(MeshBoxExtent.X, 0, MeshBoxExtent.Z));
+	MeshBoxFaces[4] = FAxisAlignedCellFace(FVector4f(0, 0, -1), FVector4f(MeshBox.Min.X, MeshBox.Min.Y, MeshBox.Min.Z), FVector4f(MeshBoxExtent.X, MeshBoxExtent.Y, 0));
+	MeshBoxFaces[5] = FAxisAlignedCellFace(FVector4f(0, 0, 1), FVector4f(MeshBox.Min.X, MeshBox.Min.Y, MeshBox.Min.Z + MeshBoxExtent.Z), FVector4f(MeshBoxExtent.X, MeshBoxExtent.Y, 0));
 
 	VisibleCellFaces.Reset();
 	VisibleCellFacePDFs.Reset();
@@ -798,7 +798,7 @@ bool ComputeBoxVisibility(
 
 	if (!bVisible)
 	{
-		const FVector4 PerpendicularVector = MeshToCellCenter ^ FVector4(0, 0, 1);
+		const FVector4f PerpendicularVector = MeshToCellCenter ^ FVector4f(0, 0, 1);
 		SamplePositions.Reset();
 
 		// Generate samples for explicit visibility sampling of the mesh
@@ -815,15 +815,15 @@ bool ComputeBoxVisibility(
 					Sample1dCDF(VisibleCellFacePDFs, VisibleCellFaceCDFs, UnnormalizedIntegral, RandomStream, PDF, Sample);
 					const int32 ChosenCellFaceIndex = FMath::TruncToInt(Sample * VisibleCellFaces.Num());
 					const FAxisAlignedCellFace& ChosenFace = CellFaces[VisibleCellFaces[ChosenCellFaceIndex]];
-					NewSample.CellPosition = ChosenFace.FaceMin + ChosenFace.FaceExtent * FVector4(RandomStream.GetFraction(), RandomStream.GetFraction(), RandomStream.GetFraction());
+					NewSample.CellPosition = ChosenFace.FaceMin + ChosenFace.FaceExtent * FVector4f(RandomStream.GetFraction(), RandomStream.GetFraction(), RandomStream.GetFraction());
 				}
 				{
 					// Generate a sample on the visible faces of the mesh
 					const int32 ChosenFaceIndex = FMath::TruncToInt(RandomStream.GetFraction() * VisibleMeshFaces.Num());
 					const FAxisAlignedCellFace& ChosenFace = MeshBoxFaces[VisibleMeshFaces[ChosenFaceIndex]];
-					NewSample.MeshPosition = ChosenFace.FaceMin + ChosenFace.FaceExtent * FVector4(RandomStream.GetFraction(), RandomStream.GetFraction(), RandomStream.GetFraction());
+					NewSample.MeshPosition = ChosenFace.FaceMin + ChosenFace.FaceExtent * FVector4f(RandomStream.GetFraction(), RandomStream.GetFraction(), RandomStream.GetFraction());
 				}
-				const FVector4 HalfPosition = .5f * (NewSample.CellPosition + NewSample.MeshPosition);
+				const FVector4f HalfPosition = .5f * (NewSample.CellPosition + NewSample.MeshPosition);
 				NewSample.PerpendicularDistance = Dot3(HalfPosition, PerpendicularVector);
 				SamplePositions.Add(NewSample);
 			}
@@ -850,8 +850,8 @@ bool ComputeBoxVisibility(
 			for (int32 MeshSampleIndex = 0; MeshSampleIndex < NumMeshSamples; MeshSampleIndex++)
 			{
 				FVisibilityQuerySample& CurrentSample = SamplePositions[CellSampleIndex * NumMeshSamples + MeshSampleIndex];
-				const FVector4 CellSamplePosition = CurrentSample.CellPosition;
-				const FVector4 MeshSamplePosition = CurrentSample.MeshPosition;
+				const FVector4f CellSamplePosition = CurrentSample.CellPosition;
+				const FVector4f MeshSamplePosition = CurrentSample.MeshPosition;
 
 				FLightRay Ray(
 					CellSamplePosition,
@@ -930,21 +930,21 @@ bool ComputeBoxVisibility(
 				const int32 SampleIndex = FMath::TruncToInt(RandomStream.GetFraction() * FurthestSamples.Num());
 				const FVisibilityQuerySample& CurrentSample = *FurthestSamples[SampleIndex];
 				const float VectorLength = (CurrentSample.CellPosition - CurrentSample.MeshPosition).Size3();
-				const FVector4 CurrentDirection = (CurrentSample.MeshPosition - CurrentSample.CellPosition).GetSafeNormal();
+				const FVector4f CurrentDirection = (CurrentSample.MeshPosition - CurrentSample.CellPosition).GetSafeNormal();
 
-				FVector4 XAxis;
-				FVector4 YAxis;
+				FVector4f XAxis;
+				FVector4f YAxis;
 				GenerateCoordinateSystem(CurrentDirection, XAxis, YAxis);
 
 				// Generate a new direction in a cone 2 degrees from the original direction, to find cracks nearby
-				const FVector4 SampleDirection = UniformSampleCone(
+				const FVector4f SampleDirection = UniformSampleCone(
 					RandomStream, 
 					FMath::Cos(2.0f * PI / 180.0f), 
 					XAxis, 
 					YAxis, 
 					CurrentDirection);
 
-				const FVector4 EndPoint = CurrentSample.CellPosition + SampleDirection * VectorLength;
+				const FVector4f EndPoint = CurrentSample.CellPosition + SampleDirection * VectorLength;
 
 				FLightRay Ray(
 					CurrentSample.CellPosition,
@@ -988,48 +988,48 @@ bool ComputeBoxVisibility(
 	return bVisible;
 }
 
-void AddMeshDebugLines(TArray<FDebugStaticLightingRay>& DebugVisibilityRays, const TArray<FStaticLightingMesh*>& Meshes, const FBox& MeshBox)
+void AddMeshDebugLines(TArray<FDebugStaticLightingRay>& DebugVisibilityRays, const TArray<FStaticLightingMesh*>& Meshes, const FBox3f& MeshBox)
 {
 	// Draw the bounding boxes of each mesh and the combined bounds
 	if (Meshes.Num() > 1)
 	{
 		for (int32 OriginalMeshIndex = 0; OriginalMeshIndex < Meshes.Num(); OriginalMeshIndex++)
 		{
-			const FVector4 Min = Meshes[OriginalMeshIndex]->BoundingBox.Min;
-			const FVector4 Max = Meshes[OriginalMeshIndex]->BoundingBox.Max;
-			DebugVisibilityRays.Add(FDebugStaticLightingRay(FVector4(Min.X, Min.Y, Min.Z), FVector4(Min.X, Max.Y, Min.Z), false));
-			DebugVisibilityRays.Add(FDebugStaticLightingRay(FVector4(Min.X, Max.Y, Min.Z), FVector4(Min.X, Max.Y, Max.Z), false));
-			DebugVisibilityRays.Add(FDebugStaticLightingRay(FVector4(Min.X, Min.Y, Min.Z), FVector4(Min.X, Min.Y, Max.Z), false));
-			DebugVisibilityRays.Add(FDebugStaticLightingRay(FVector4(Min.X, Min.Y, Max.Z), FVector4(Min.X, Max.Y, Max.Z), false));
-
-			DebugVisibilityRays.Add(FDebugStaticLightingRay(FVector4(Max.X, Min.Y, Min.Z), FVector4(Max.X, Max.Y, Min.Z), false));
-			DebugVisibilityRays.Add(FDebugStaticLightingRay(FVector4(Max.X, Max.Y, Min.Z), FVector4(Max.X, Max.Y, Max.Z), false));
-			DebugVisibilityRays.Add(FDebugStaticLightingRay(FVector4(Max.X, Min.Y, Min.Z), FVector4(Max.X, Min.Y, Max.Z), false));
-			DebugVisibilityRays.Add(FDebugStaticLightingRay(FVector4(Max.X, Min.Y, Max.Z), FVector4(Max.X, Max.Y, Max.Z), false));
-
-			DebugVisibilityRays.Add(FDebugStaticLightingRay(FVector4(Min.X, Min.Y, Min.Z), FVector4(Max.X, Min.Y, Min.Z), false));
-			DebugVisibilityRays.Add(FDebugStaticLightingRay(FVector4(Min.X, Min.Y, Max.Z), FVector4(Max.X, Min.Y, Max.Z), false));
-			DebugVisibilityRays.Add(FDebugStaticLightingRay(FVector4(Min.X, Max.Y, Min.Z), FVector4(Max.X, Max.Y, Min.Z), false));
-			DebugVisibilityRays.Add(FDebugStaticLightingRay(FVector4(Min.X, Max.Y, Max.Z), FVector4(Max.X, Max.Y, Max.Z), false));
+			const FVector4f Min = Meshes[OriginalMeshIndex]->BoundingBox.Min;
+			const FVector4f Max = Meshes[OriginalMeshIndex]->BoundingBox.Max;
+			DebugVisibilityRays.Add(FDebugStaticLightingRay(FVector4f(Min.X, Min.Y, Min.Z), FVector4f(Min.X, Max.Y, Min.Z), false));
+			DebugVisibilityRays.Add(FDebugStaticLightingRay(FVector4f(Min.X, Max.Y, Min.Z), FVector4f(Min.X, Max.Y, Max.Z), false));
+			DebugVisibilityRays.Add(FDebugStaticLightingRay(FVector4f(Min.X, Min.Y, Min.Z), FVector4f(Min.X, Min.Y, Max.Z), false));
+			DebugVisibilityRays.Add(FDebugStaticLightingRay(FVector4f(Min.X, Min.Y, Max.Z), FVector4f(Min.X, Max.Y, Max.Z), false));
+
+			DebugVisibilityRays.Add(FDebugStaticLightingRay(FVector4f(Max.X, Min.Y, Min.Z), FVector4f(Max.X, Max.Y, Min.Z), false));
+			DebugVisibilityRays.Add(FDebugStaticLightingRay(FVector4f(Max.X, Max.Y, Min.Z), FVector4f(Max.X, Max.Y, Max.Z), false));
+			DebugVisibilityRays.Add(FDebugStaticLightingRay(FVector4f(Max.X, Min.Y, Min.Z), FVector4f(Max.X, Min.Y, Max.Z), false));
+			DebugVisibilityRays.Add(FDebugStaticLightingRay(FVector4f(Max.X, Min.Y, Max.Z), FVector4f(Max.X, Max.Y, Max.Z), false));
+
+			DebugVisibilityRays.Add(FDebugStaticLightingRay(FVector4f(Min.X, Min.Y, Min.Z), FVector4f(Max.X, Min.Y, Min.Z), false));
+			DebugVisibilityRays.Add(FDebugStaticLightingRay(FVector4f(Min.X, Min.Y, Max.Z), FVector4f(Max.X, Min.Y, Max.Z), false));
+			DebugVisibilityRays.Add(FDebugStaticLightingRay(FVector4f(Min.X, Max.Y, Min.Z), FVector4f(Max.X, Max.Y, Min.Z), false));
+			DebugVisibilityRays.Add(FDebugStaticLightingRay(FVector4f(Min.X, Max.Y, Max.Z), FVector4f(Max.X, Max.Y, Max.Z), false));
 		}
 	}
 
-	const FVector4 Min = MeshBox.Min;
-	const FVector4 Max = MeshBox.Max;
-	DebugVisibilityRays.Add(FDebugStaticLightingRay(FVector4(Min.X, Min.Y, Min.Z), FVector4(Min.X, Max.Y, Min.Z), true));
-	DebugVisibilityRays.Add(FDebugStaticLightingRay(FVector4(Min.X, Max.Y, Min.Z), FVector4(Min.X, Max.Y, Max.Z), true));
-	DebugVisibilityRays.Add(FDebugStaticLightingRay(FVector4(Min.X, Min.Y, Min.Z), FVector4(Min.X, Min.Y, Max.Z), true));
-	DebugVisibilityRays.Add(FDebugStaticLightingRay(FVector4(Min.X, Min.Y, Max.Z), FVector4(Min.X, Max.Y, Max.Z), true));
-
-	DebugVisibilityRays.Add(FDebugStaticLightingRay(FVector4(Max.X, Min.Y, Min.Z), FVector4(Max.X, Max.Y, Min.Z), true));
-	DebugVisibilityRays.Add(FDebugStaticLightingRay(FVector4(Max.X, Max.Y, Min.Z), FVector4(Max.X, Max.Y, Max.Z), true));
-	DebugVisibilityRays.Add(FDebugStaticLightingRay(FVector4(Max.X, Min.Y, Min.Z), FVector4(Max.X, Min.Y, Max.Z), true));
-	DebugVisibilityRays.Add(FDebugStaticLightingRay(FVector4(Max.X, Min.Y, Max.Z), FVector4(Max.X, Max.Y, Max.Z), true));
-
-	DebugVisibilityRays.Add(FDebugStaticLightingRay(FVector4(Min.X, Min.Y, Min.Z), FVector4(Max.X, Min.Y, Min.Z), true));
-	DebugVisibilityRays.Add(FDebugStaticLightingRay(FVector4(Min.X, Min.Y, Max.Z), FVector4(Max.X, Min.Y, Max.Z), true));
-	DebugVisibilityRays.Add(FDebugStaticLightingRay(FVector4(Min.X, Max.Y, Min.Z), FVector4(Max.X, Max.Y, Min.Z), true));
-	DebugVisibilityRays.Add(FDebugStaticLightingRay(FVector4(Min.X, Max.Y, Max.Z), FVector4(Max.X, Max.Y, Max.Z), true));
+	const FVector4f Min = MeshBox.Min;
+	const FVector4f Max = MeshBox.Max;
+	DebugVisibilityRays.Add(FDebugStaticLightingRay(FVector4f(Min.X, Min.Y, Min.Z), FVector4f(Min.X, Max.Y, Min.Z), true));
+	DebugVisibilityRays.Add(FDebugStaticLightingRay(FVector4f(Min.X, Max.Y, Min.Z), FVector4f(Min.X, Max.Y, Max.Z), true));
+	DebugVisibilityRays.Add(FDebugStaticLightingRay(FVector4f(Min.X, Min.Y, Min.Z), FVector4f(Min.X, Min.Y, Max.Z), true));
+	DebugVisibilityRays.Add(FDebugStaticLightingRay(FVector4f(Min.X, Min.Y, Max.Z), FVector4f(Min.X, Max.Y, Max.Z), true));
+
+	DebugVisibilityRays.Add(FDebugStaticLightingRay(FVector4f(Max.X, Min.Y, Min.Z), FVector4f(Max.X, Max.Y, Min.Z), true));
+	DebugVisibilityRays.Add(FDebugStaticLightingRay(FVector4f(Max.X, Max.Y, Min.Z), FVector4f(Max.X, Max.Y, Max.Z), true));
+	DebugVisibilityRays.Add(FDebugStaticLightingRay(FVector4f(Max.X, Min.Y, Min.Z), FVector4f(Max.X, Min.Y, Max.Z), true));
+	DebugVisibilityRays.Add(FDebugStaticLightingRay(FVector4f(Max.X, Min.Y, Max.Z), FVector4f(Max.X, Max.Y, Max.Z), true));
+
+	DebugVisibilityRays.Add(FDebugStaticLightingRay(FVector4f(Min.X, Min.Y, Min.Z), FVector4f(Max.X, Min.Y, Min.Z), true));
+	DebugVisibilityRays.Add(FDebugStaticLightingRay(FVector4f(Min.X, Min.Y, Max.Z), FVector4f(Max.X, Min.Y, Max.Z), true));
+	DebugVisibilityRays.Add(FDebugStaticLightingRay(FVector4f(Min.X, Max.Y, Min.Z), FVector4f(Max.X, Max.Y, Min.Z), true));
+	DebugVisibilityRays.Add(FDebugStaticLightingRay(FVector4f(Min.X, Max.Y, Max.Z), FVector4f(Max.X, Max.Y, Max.Z), true));
 }
 
 /** Calculates visibility for a given group of cells, called from all threads. */
@@ -1090,13 +1090,13 @@ void FStaticLightingSystem::CalculatePrecomputedVisibility(int32 BucketIndex)
 		CurrentCell.VisibilityData.AddZeroed(VisibilityMeshes.Num() / 8 + 1);
 
 		FAxisAlignedCellFace CellFaces[6];
-		const FVector CellBoundsSize = CurrentCell.Bounds.GetSize();
-		CellFaces[0] = FAxisAlignedCellFace(FVector4(-1, 0, 0), FVector4(CurrentCell.Bounds.Min.X, CurrentCell.Bounds.Min.Y, CurrentCell.Bounds.Min.Z), FVector4(0, CellBoundsSize.Y, CellBoundsSize.Z));
-		CellFaces[1] = FAxisAlignedCellFace(FVector4(1, 0, 0), FVector4(CurrentCell.Bounds.Max.X, CurrentCell.Bounds.Min.Y, CurrentCell.Bounds.Min.Z), FVector4(0, CellBoundsSize.Y, CellBoundsSize.Z));
-		CellFaces[2] = FAxisAlignedCellFace(FVector4(0, -1, 0), FVector4(CurrentCell.Bounds.Min.X, CurrentCell.Bounds.Min.Y, CurrentCell.Bounds.Min.Z), FVector4(CellBoundsSize.X, 0, CellBoundsSize.Z));
-		CellFaces[3] = FAxisAlignedCellFace(FVector4(0, 1, 0), FVector4(CurrentCell.Bounds.Min.X, CurrentCell.Bounds.Max.Y, CurrentCell.Bounds.Min.Z), FVector4(CellBoundsSize.X, 0, CellBoundsSize.Z));
-		CellFaces[4] = FAxisAlignedCellFace(FVector4(0, 0, -1), FVector4(CurrentCell.Bounds.Min.X, CurrentCell.Bounds.Min.Y, CurrentCell.Bounds.Min.Z), FVector4(CellBoundsSize.X, CellBoundsSize.Y, 0));
-		CellFaces[5] = FAxisAlignedCellFace(FVector4(0, 0, 1), FVector4(CurrentCell.Bounds.Min.X, CurrentCell.Bounds.Min.Y, CurrentCell.Bounds.Max.Z), FVector4(CellBoundsSize.X, CellBoundsSize.Y, 0));
+		const FVector3f CellBoundsSize = CurrentCell.Bounds.GetSize();
+		CellFaces[0] = FAxisAlignedCellFace(FVector4f(-1, 0, 0), FVector4f(CurrentCell.Bounds.Min.X, CurrentCell.Bounds.Min.Y, CurrentCell.Bounds.Min.Z), FVector4f(0, CellBoundsSize.Y, CellBoundsSize.Z));
+		CellFaces[1] = FAxisAlignedCellFace(FVector4f(1, 0, 0), FVector4f(CurrentCell.Bounds.Max.X, CurrentCell.Bounds.Min.Y, CurrentCell.Bounds.Min.Z), FVector4f(0, CellBoundsSize.Y, CellBoundsSize.Z));
+		CellFaces[2] = FAxisAlignedCellFace(FVector4f(0, -1, 0), FVector4f(CurrentCell.Bounds.Min.X, CurrentCell.Bounds.Min.Y, CurrentCell.Bounds.Min.Z), FVector4f(CellBoundsSize.X, 0, CellBoundsSize.Z));
+		CellFaces[3] = FAxisAlignedCellFace(FVector4f(0, 1, 0), FVector4f(CurrentCell.Bounds.Min.X, CurrentCell.Bounds.Max.Y, CurrentCell.Bounds.Min.Z), FVector4f(CellBoundsSize.X, 0, CellBoundsSize.Z));
+		CellFaces[4] = FAxisAlignedCellFace(FVector4f(0, 0, -1), FVector4f(CurrentCell.Bounds.Min.X, CurrentCell.Bounds.Min.Y, CurrentCell.Bounds.Min.Z), FVector4f(CellBoundsSize.X, CellBoundsSize.Y, 0));
+		CellFaces[5] = FAxisAlignedCellFace(FVector4f(0, 0, 1), FVector4f(CurrentCell.Bounds.Min.X, CurrentCell.Bounds.Min.Y, CurrentCell.Bounds.Max.Z), FVector4f(CellBoundsSize.X, CellBoundsSize.Y, 0));
 
 		GroupVisibility.Reset();
 		
@@ -1132,7 +1132,7 @@ void FStaticLightingSystem::CalculatePrecomputedVisibility(int32 BucketIndex)
 		{
 			const FVisibilityMesh& VisibilityMesh = VisibilityMeshes[VisibilityMeshIndex];
 
-			FBox OriginalMeshBounds(ForceInit);
+			FBox3f OriginalMeshBounds(ForceInit);
 			// Combine mesh bounds, usually only BSP has multiple meshes per VisibilityId
 			//@todo - could explicitly sample each bounds separately, but they tend to be pretty close together in world space
 			for (int32 OriginalMeshIndex = 0; OriginalMeshIndex < VisibilityMesh.Meshes.Num(); OriginalMeshIndex++)
@@ -1140,7 +1140,7 @@ void FStaticLightingSystem::CalculatePrecomputedVisibility(int32 BucketIndex)
 				OriginalMeshBounds += VisibilityMesh.Meshes[OriginalMeshIndex]->BoundingBox;
 			}
 
-			const FBox MeshBox(
+			const FBox3f MeshBox(
 				OriginalMeshBounds.GetCenter() - OriginalMeshBounds.GetExtent() * PrecomputedVisibilitySettings.MeshBoundsScale,
 				OriginalMeshBounds.GetCenter() + OriginalMeshBounds.GetExtent() * PrecomputedVisibilitySettings.MeshBoundsScale);
 
@@ -1223,22 +1223,22 @@ void FStaticLightingSystem::CalculatePrecomputedVisibility(int32 BucketIndex)
 		if (bDebugThisCell)
 		{
 			// Draw the bounds of each cell processed
-			const FVector4 Min = CurrentCell.Bounds.Min;
-			const FVector4 Max = CurrentCell.Bounds.Max;
-			DataLink->Element.DebugVisibilityRays.Add(FDebugStaticLightingRay(FVector4(Min.X, Min.Y, Min.Z), FVector4(Min.X, Max.Y, Min.Z), false));
-			DataLink->Element.DebugVisibilityRays.Add(FDebugStaticLightingRay(FVector4(Min.X, Max.Y, Min.Z), FVector4(Min.X, Max.Y, Max.Z), false));
-			DataLink->Element.DebugVisibilityRays.Add(FDebugStaticLightingRay(FVector4(Min.X, Min.Y, Min.Z), FVector4(Min.X, Min.Y, Max.Z), false));
-			DataLink->Element.DebugVisibilityRays.Add(FDebugStaticLightingRay(FVector4(Min.X, Min.Y, Max.Z), FVector4(Min.X, Max.Y, Max.Z), false));
-
-			DataLink->Element.DebugVisibilityRays.Add(FDebugStaticLightingRay(FVector4(Max.X, Min.Y, Min.Z), FVector4(Max.X, Max.Y, Min.Z), false));
-			DataLink->Element.DebugVisibilityRays.Add(FDebugStaticLightingRay(FVector4(Max.X, Max.Y, Min.Z), FVector4(Max.X, Max.Y, Max.Z), false));
-			DataLink->Element.DebugVisibilityRays.Add(FDebugStaticLightingRay(FVector4(Max.X, Min.Y, Min.Z), FVector4(Max.X, Min.Y, Max.Z), false));
-			DataLink->Element.DebugVisibilityRays.Add(FDebugStaticLightingRay(FVector4(Max.X, Min.Y, Max.Z), FVector4(Max.X, Max.Y, Max.Z), false));
-
-			DataLink->Element.DebugVisibilityRays.Add(FDebugStaticLightingRay(FVector4(Min.X, Min.Y, Min.Z), FVector4(Max.X, Min.Y, Min.Z), false));
-			DataLink->Element.DebugVisibilityRays.Add(FDebugStaticLightingRay(FVector4(Min.X, Min.Y, Max.Z), FVector4(Max.X, Min.Y, Max.Z), false));
-			DataLink->Element.DebugVisibilityRays.Add(FDebugStaticLightingRay(FVector4(Min.X, Max.Y, Min.Z), FVector4(Max.X, Max.Y, Min.Z), false));
-			DataLink->Element.DebugVisibilityRays.Add(FDebugStaticLightingRay(FVector4(Min.X, Max.Y, Max.Z), FVector4(Max.X, Max.Y, Max.Z), false));
+			const FVector4f Min = CurrentCell.Bounds.Min;
+			const FVector4f Max = CurrentCell.Bounds.Max;
+			DataLink->Element.DebugVisibilityRays.Add(FDebugStaticLightingRay(FVector4f(Min.X, Min.Y, Min.Z), FVector4f(Min.X, Max.Y, Min.Z), false));
+			DataLink->Element.DebugVisibilityRays.Add(FDebugStaticLightingRay(FVector4f(Min.X, Max.Y, Min.Z), FVector4f(Min.X, Max.Y, Max.Z), false));
+			DataLink->Element.DebugVisibilityRays.Add(FDebugStaticLightingRay(FVector4f(Min.X, Min.Y, Min.Z), FVector4f(Min.X, Min.Y, Max.Z), false));
+			DataLink->Element.DebugVisibilityRays.Add(FDebugStaticLightingRay(FVector4f(Min.X, Min.Y, Max.Z), FVector4f(Min.X, Max.Y, Max.Z), false));
+
+			DataLink->Element.DebugVisibilityRays.Add(FDebugStaticLightingRay(FVector4f(Max.X, Min.Y, Min.Z), FVector4f(Max.X, Max.Y, Min.Z), false));
+			DataLink->Element.DebugVisibilityRays.Add(FDebugStaticLightingRay(FVector4f(Max.X, Max.Y, Min.Z), FVector4f(Max.X, Max.Y, Max.Z), false));
+			DataLink->Element.DebugVisibilityRays.Add(FDebugStaticLightingRay(FVector4f(Max.X, Min.Y, Min.Z), FVector4f(Max.X, Min.Y, Max.Z), false));
+			DataLink->Element.DebugVisibilityRays.Add(FDebugStaticLightingRay(FVector4f(Max.X, Min.Y, Max.Z), FVector4f(Max.X, Max.Y, Max.Z), false));
+
+			DataLink->Element.DebugVisibilityRays.Add(FDebugStaticLightingRay(FVector4f(Min.X, Min.Y, Min.Z), FVector4f(Max.X, Min.Y, Min.Z), false));
+			DataLink->Element.DebugVisibilityRays.Add(FDebugStaticLightingRay(FVector4f(Min.X, Min.Y, Max.Z), FVector4f(Max.X, Min.Y, Max.Z), false));
+			DataLink->Element.DebugVisibilityRays.Add(FDebugStaticLightingRay(FVector4f(Min.X, Max.Y, Min.Z), FVector4f(Max.X, Max.Y, Min.Z), false));
+			DataLink->Element.DebugVisibilityRays.Add(FDebugStaticLightingRay(FVector4f(Min.X, Max.Y, Max.Z), FVector4f(Max.X, Max.Y, Max.Z), false));
 		}
 	}
 	
diff --git a/Engine/Source/Programs/UnrealLightmass/Private/Lighting/Radiosity.cpp b/Engine/Source/Programs/UnrealLightmass/Private/Lighting/Radiosity.cpp
--- a/Engine/Source/Programs/UnrealLightmass/Private/Lighting/Radiosity.cpp
+++ b/Engine/Source/Programs/UnrealLightmass/Private/Lighting/Radiosity.cpp
@@ -234,7 +234,7 @@ void FStaticLightingSystem::RadiositySetupTextureMapping(FStaticLightingTextureM
 	checkSlow(TextureMapping);
 	FStaticLightingMappingContext MappingContext(TextureMapping->Mesh,*this,&StartupDebugOutput);
 	LIGHTINGSTAT(FScopedRDTSCTimer CachingTime(MappingContext.Stats.RadiositySetupThreadTime));
-	const FBoxSphereBounds ImportanceBounds = Scene.GetImportanceBounds();
+	const FBoxSphereBounds3f ImportanceBounds = Scene.GetImportanceBounds();
 
 	FTexelToVertexMap TexelToVertexMap(TextureMapping->SurfaceCacheSizeX, TextureMapping->SurfaceCacheSizeY);
 
@@ -303,7 +303,7 @@ void FStaticLightingSystem::RadiositySetupTextureMapping(FStaticLightingTextureM
 					if (!RadiosityCache.InterpolateLighting(Vertex, true, false, 1.0f, SkyLighting, UnusedSecondLighting, UnusedBackfacingHitsFraction, MappingContext.DebugCacheRecords))
 					{
 						FFinalGatherSample UniformSampledIncomingRadiance;
-						TArray<FVector4> ImportancePhotonDirections;
+						TArray<FVector4f> ImportancePhotonDirections;
 						FLightingCacheGatherInfo GatherInfo;
 						const int32 NumAdaptiveRefinementLevels = GeneralSettings.IndirectLightingQuality <= 10 ? 1 : 2;
 
@@ -345,8 +345,8 @@ void FStaticLightingSystem::RadiositySetupTextureMapping(FStaticLightingTextureM
 							IrradianceCachingSettings,
 							GeneralSettings,
 							UniformSampledIncomingRadiance,
-							FVector4(0, 0, 0, 0),
-							FVector4(0, 0, 0, 0)
+							FVector4f(0, 0, 0, 0),
+							FVector4f(0, 0, 0, 0)
 							);
 
 						// Add the incident radiance sample to the cache.
@@ -420,8 +420,8 @@ void FStaticLightingSystem::RadiositySetupTextureMapping(FStaticLightingTextureM
 					FGatheredLightSample DirectLighting;
 					FGatheredLightSample Unused;
 					float Unused2;
-					TArray<FVector, TInlineAllocator<1>> VertexOffsets;
-					VertexOffsets.Add(FVector(0, 0, 0));
+					TArray<FVector3f, TInlineAllocator<1>> VertexOffsets;
+					VertexOffsets.Add(FVector3f(0, 0, 0));
 
 					CalculateApproximateDirectLighting(CurrentVertex, TexelToVertex.TexelRadius, VertexOffsets, .1f, true, true, bDebugThisTexel, MappingContext, DirectLighting, Unused, Unused2);
 
@@ -554,7 +554,7 @@ void FStaticLightingSystem::RadiosityIterationTextureMapping(FStaticLightingText
 	checkSlow(TextureMapping);
 	FStaticLightingMappingContext MappingContext(TextureMapping->Mesh,*this,&StartupDebugOutput);
 	LIGHTINGSTAT(FScopedRDTSCTimer CachingTime(MappingContext.Stats.RadiosityIterationThreadTime));
-	const FBoxSphereBounds ImportanceBounds = Scene.GetImportanceBounds();
+	const FBoxSphereBounds3f ImportanceBounds = Scene.GetImportanceBounds();
 
 	FTexelToVertexMap TexelToVertexMap(TextureMapping->SurfaceCacheSizeX, TextureMapping->SurfaceCacheSizeY);
 
@@ -620,7 +620,7 @@ void FStaticLightingSystem::RadiosityIterationTextureMapping(FStaticLightingText
 					{
 						FFinalGatherSample UniformSampledIncomingRadiance;
 						//@todo - find and pass in photons from the appropriate bounce number to improve bUseRadiositySolverForLightMultibounce quality
-						TArray<FVector4> ImportancePhotonDirections;
+						TArray<FVector4f> ImportancePhotonDirections;
 						FLightingCacheGatherInfo GatherInfo;
 
 						// Gather previous iteration results (double buffer)
@@ -658,8 +658,8 @@ void FStaticLightingSystem::RadiosityIterationTextureMapping(FStaticLightingText
 							IrradianceCachingSettings,
 							GeneralSettings,
 							UniformSampledIncomingRadiance,
-							FVector4(0, 0, 0, 0),
-							FVector4(0, 0, 0, 0)
+							FVector4f(0, 0, 0, 0),
+							FVector4f(0, 0, 0, 0)
 							);
 
 						// Add the incident radiance sample to the cache.
diff --git a/Engine/Source/Programs/UnrealLightmass/Private/Lighting/Raster.h b/Engine/Source/Programs/UnrealLightmass/Private/Lighting/Raster.h
--- a/Engine/Source/Programs/UnrealLightmass/Private/Lighting/Raster.h
+++ b/Engine/Source/Programs/UnrealLightmass/Private/Lighting/Raster.h
@@ -23,10 +23,10 @@ public:
 
 	typedef typename RasterPolicyType::InterpolantType InterpolantType;
 
-	void DrawTriangle(const InterpolantType& I0,const InterpolantType& I1,const InterpolantType& I2,const FVector2D& P0,const FVector2D& P1,const FVector2D& P2,bool BackFacing)
+	void DrawTriangle(const InterpolantType& I0,const InterpolantType& I1,const InterpolantType& I2,const FVector2f& P0,const FVector2f& P1,const FVector2f& P2,bool BackFacing)
 	{
 		InterpolantType	Interpolants[3] = { I0, I1, I2 };
-		FVector2D		Points[3] = { P0, P1, P2 };
+		FVector2f		Points[3] = { P0, P1, P2 };
 
 		// Find the top point.
 
diff --git a/Engine/Source/Programs/UnrealLightmass/Private/Lighting/SampleVolume.cpp b/Engine/Source/Programs/UnrealLightmass/Private/Lighting/SampleVolume.cpp
--- a/Engine/Source/Programs/UnrealLightmass/Private/Lighting/SampleVolume.cpp
+++ b/Engine/Source/Programs/UnrealLightmass/Private/Lighting/SampleVolume.cpp
@@ -26,7 +26,7 @@ struct FVolumeLightingProximityOctreeSemantics
 	static FBoxCenterAndExtent GetBoundingBox(const FVolumeSampleProximityElement& Element)
 	{
 		const FVolumeLightingSample& Sample = Element.VolumeSamples[Element.SampleIndex];
-		return FBoxCenterAndExtent(FVector4(Sample.PositionAndRadius, 0.0f), FVector4(0,0,0));
+		return FBoxCenterAndExtent(FVector4f(Sample.PositionAndRadius, 0.0f), FVector4f(0,0,0));
 	}
 };
 
@@ -52,9 +52,9 @@ void FVolumeLightingSample::ToSHVector(FSHVectorRGB3& SHVector) const
 }
 
 /** Returns true if there is an existing sample in VolumeOctree within SearchDistance of Position. */
-static bool FindNearbyVolumeSample(const FVolumeLightingProximityOctree& VolumeOctree, const FVector4& Position, float SearchDistance)
+static bool FindNearbyVolumeSample(const FVolumeLightingProximityOctree& VolumeOctree, const FVector4f& Position, float SearchDistance)
 {
-	const FBox SearchBox = FBox::BuildAABB(Position, FVector4(SearchDistance, SearchDistance, SearchDistance));
+	const FBox3f SearchBox = FBox3f::BuildAABB(Position, FVector4f(SearchDistance, SearchDistance, SearchDistance));
 	for (FVolumeLightingProximityOctree::TConstIterator<> OctreeIt(VolumeOctree); OctreeIt.HasPendingNodes(); OctreeIt.Advance())
 	{
 		const FVolumeLightingProximityOctree::FNode& CurrentNode = OctreeIt.GetCurrentNode();
@@ -125,7 +125,7 @@ public:
 		}
 
 		
-		TArray<FVector2D> UniformHemisphereSampleUniforms;
+		TArray<FVector2f> UniformHemisphereSampleUniforms;
 		const int32 NumUpperVolumeSamples = 16;
 		const float NumThetaStepsFloat = FMath::Sqrt(NumUpperVolumeSamples / (float)PI);
 		const int32 NumThetaSteps = FMath::TruncToInt(NumThetaStepsFloat);
@@ -163,21 +163,21 @@ private:
 	FCoherentRayCache& CoherentRayCache;
 	FVolumeLightingProximityOctree& ProximityOctree;
 	TArray<float> LayerHeightOffsets;
-	TArray<FVector4> UniformHemisphereSamples;
+	TArray<FVector4f> UniformHemisphereSamples;
 };
 
 int32 ComputeNumBackfacingHits(
-	const FVector4& SamplePosition, 
+	const FVector4f& SamplePosition, 
 	FStaticLightingSystem& System, 
 	FCoherentRayCache& CoherentRayCache, 
 	float SceneBoundingRadius, 
-	const TArray<FVector4>& UniformHemisphereSamples)
+	const TArray<FVector4f>& UniformHemisphereSamples)
 {
 	int NumBackfacingHits = 0;
 
 	for (int32 SampleIndex = 0; SampleIndex < UniformHemisphereSamples.Num() * 2; SampleIndex++)
 	{
-		FVector4 SampleDirection = UniformHemisphereSamples[SampleIndex % UniformHemisphereSamples.Num()];
+		FVector4f SampleDirection = UniformHemisphereSamples[SampleIndex % UniformHemisphereSamples.Num()];
 		SampleDirection.Z *= SampleIndex >= UniformHemisphereSamples.Num() ? -1.0f : 1.0f;
 
 		const FLightRay PathRay(
@@ -207,7 +207,7 @@ void FVolumeSamplePlacementRasterPolicy::ProcessPixel(int32 X,int32 Y,const Inte
 		// Place a sample for each layer
 		for (int32 LayerIndex = 0; LayerIndex < LayerHeightOffsets.Num(); LayerIndex++)
 		{
-			const FVector4 SamplePosition = Vertex.WorldPosition + FVector4(0, 0, LayerHeightOffsets[LayerIndex]);
+			const FVector4f SamplePosition = Vertex.WorldPosition + FVector4f(0, 0, LayerHeightOffsets[LayerIndex]);
 			// Only place a sample if there isn't already one nearby, to avoid clumping
 			if (!FindNearbyVolumeSample(ProximityOctree, SamplePosition, MinSampleDistance))
 			{
@@ -219,7 +219,7 @@ void FVolumeSamplePlacementRasterPolicy::ProcessPixel(int32 X,int32 Y,const Inte
 					TArray<FVolumeLightingSample>* VolumeLightingSamples = System.VolumeLightingSamples.Find(LevelGuid);
 					check(VolumeLightingSamples);
 					// Add a new sample for this layer
-					VolumeLightingSamples->Add(FVolumeLightingSample(FVector4(SamplePosition, SampleRadius)));
+					VolumeLightingSamples->Add(FVolumeLightingSample(FVector4f(SamplePosition, SampleRadius)));
 					// Add the sample to the proximity octree so we can avoid placing any more samples nearby
 					ProximityOctree.AddElement(FVolumeSampleProximityElement(VolumeLightingSamples->Num() - 1, *VolumeLightingSamples));
 					if (System.DynamicObjectSettings.bVisualizeVolumeLightInterpolation)
@@ -239,7 +239,7 @@ void FStaticLightingSystem::BeginCalculateVolumeSamples()
 	VolumeBounds = GetImportanceBounds(false);
 	if (VolumeBounds.SphereRadius < DELTA)
 	{
-		VolumeBounds = FBoxSphereBounds(AggregateMesh->GetBounds());
+		VolumeBounds = FBoxSphereBounds3f(AggregateMesh->GetBounds());
 	}
 
 	// Only place samples if the volume has area
@@ -261,7 +261,7 @@ void FStaticLightingSystem::BeginCalculateVolumeSamples()
 				// Only place inside the importance volume
 				if (IsPointInImportanceVolume(Vertices[0].WorldPosition))
 				{
-					FVector4 TriangleNormal = (Vertices[2].WorldPosition - Vertices[0].WorldPosition) ^ (Vertices[1].WorldPosition - Vertices[0].WorldPosition);
+					FVector4f TriangleNormal = (Vertices[2].WorldPosition - Vertices[0].WorldPosition) ^ (Vertices[1].WorldPosition - Vertices[0].WorldPosition);
 					TriangleNormal.Z = 0.f; // approximate only for X-Y plane
 					float TotalArea = 0.5f * TriangleNormal.Size3() * CurrentMesh->NumTriangles;
 					LandscapeEstimateNum += TotalArea / FMath::Square(DynamicObjectSettings.SurfaceLightSampleSpacing);
@@ -301,7 +301,7 @@ void FStaticLightingSystem::BeginCalculateVolumeSamples()
 			RasterSizeY, 
 			// Use a minimum sample distance slightly less than the SurfaceLightSampleSpacing
 			0.9f * FMath::Min(DynamicObjectSettings.SurfaceLightSampleSpacing, DynamicObjectSettings.SurfaceSampleLayerHeightSpacing), 
-			FBoxSphereBounds(AggregateMesh->GetBounds()).SphereRadius,
+			FBoxSphereBounds3f(AggregateMesh->GetBounds()).SphereRadius,
 			SampleRadius,
 			*this,
 			MappingContext.RayCache,
@@ -337,16 +337,16 @@ void FStaticLightingSystem::BeginCalculateVolumeSamples()
 
 					if (CurrentMesh->IsElementCastingShadow(ElementIndex))
 					{
-						FVector2D XYPositions[3];
+						FVector2f XYPositions[3];
 						for (int32 VertIndex = 0; VertIndex < 3; VertIndex++)
 						{
 							// Transform world space positions from [VolumeBounds.Origin - VolumeBounds.BoxExtent, VolumeBounds.Origin + VolumeBounds.BoxExtent] into [0,1]
-							const FVector4 TransformedPosition = (Vertices[VertIndex].WorldPosition - FVector4(VolumeBounds.Origin, 0.0f) + FVector4(VolumeBounds.BoxExtent, 0.0f)) / (2.0f * FVector4(VolumeBounds.BoxExtent, 1.0f));
+							const FVector4f TransformedPosition = (Vertices[VertIndex].WorldPosition - FVector4f(VolumeBounds.Origin, 0.0f) + FVector4f(VolumeBounds.BoxExtent, 0.0f)) / (2.0f * FVector4f(VolumeBounds.BoxExtent, 1.0f));
 							// Project positions onto the XY plane and scale to the resolution determined by DynamicObjectSettings.SurfaceLightSampleSpacing
-							XYPositions[VertIndex] = FVector2D(TransformedPosition.X * RasterSizeX, TransformedPosition.Y * RasterSizeY);
+							XYPositions[VertIndex] = FVector2f(TransformedPosition.X * RasterSizeX, TransformedPosition.Y * RasterSizeY);
 						}
 
-						const FVector4 TriangleNormal = (Vertices[2].WorldPosition - Vertices[0].WorldPosition) ^ (Vertices[1].WorldPosition - Vertices[0].WorldPosition);
+						const FVector4f TriangleNormal = (Vertices[2].WorldPosition - Vertices[0].WorldPosition) ^ (Vertices[1].WorldPosition - Vertices[0].WorldPosition);
 						const float TriangleArea = 0.5f * TriangleNormal.Size3();
 
 						if (TriangleArea > DELTA)
@@ -354,9 +354,9 @@ void FStaticLightingSystem::BeginCalculateVolumeSamples()
 							if (TextureMapping)
 							{
 								// Triangle vertices in lightmap UV space, scaled by the lightmap resolution
-								const FVector2D Vertex0 = Vertices[0].TextureCoordinates[TextureMapping->LightmapTextureCoordinateIndex] * FVector2D(TextureMapping->SizeX, TextureMapping->SizeY);
-								const FVector2D Vertex1 = Vertices[1].TextureCoordinates[TextureMapping->LightmapTextureCoordinateIndex] * FVector2D(TextureMapping->SizeX, TextureMapping->SizeY);
-								const FVector2D Vertex2 = Vertices[2].TextureCoordinates[TextureMapping->LightmapTextureCoordinateIndex] * FVector2D(TextureMapping->SizeX, TextureMapping->SizeY);
+								const FVector2f Vertex0 = Vertices[0].TextureCoordinates[TextureMapping->LightmapTextureCoordinateIndex] * FVector2f(TextureMapping->SizeX, TextureMapping->SizeY);
+								const FVector2f Vertex1 = Vertices[1].TextureCoordinates[TextureMapping->LightmapTextureCoordinateIndex] * FVector2f(TextureMapping->SizeX, TextureMapping->SizeY);
+								const FVector2f Vertex2 = Vertices[2].TextureCoordinates[TextureMapping->LightmapTextureCoordinateIndex] * FVector2f(TextureMapping->SizeX, TextureMapping->SizeY);
 
 								// Area in lightmap space, or the number of lightmap texels covered by this triangle
 								const float LightmapTriangleArea = FMath::Abs(
@@ -396,21 +396,21 @@ void FStaticLightingSystem::BeginCalculateVolumeSamples()
 		// Generate samples in a uniform 3d grid inside the detail volumes.  These will handle detail indirect lighting in areas that aren't directly above a surface.
 		for (int32 VolumeIndex = 0; VolumeIndex < Scene.CharacterIndirectDetailVolumes.Num(); VolumeIndex++)
 		{
-			const FBox& DetailVolumeBounds = Scene.CharacterIndirectDetailVolumes[VolumeIndex];
+			const FBox3f& DetailVolumeBounds = Scene.CharacterIndirectDetailVolumes[VolumeIndex];
 			for (float SampleX = DetailVolumeBounds.Min.X; SampleX < DetailVolumeBounds.Max.X + DetailVolumeSpacing; SampleX += DetailVolumeSpacing)
 			{
 				for (float SampleY = DetailVolumeBounds.Min.Y; SampleY < DetailVolumeBounds.Max.Y + DetailVolumeSpacing; SampleY += DetailVolumeSpacing)
 				{
 					for (float SampleZ = DetailVolumeBounds.Min.Z; SampleZ < DetailVolumeBounds.Max.Z + DetailVolumeSpacing; SampleZ += DetailVolumeSpacing)
 					{
-						const FVector4 SamplePosition(SampleX, SampleY, SampleZ);
+						const FVector4f SamplePosition(SampleX, SampleY, SampleZ);
 							
 						// Only place a sample if there are no surface lighting samples nearby
 						if (!FindNearbyVolumeSample(VolumeLightingOctree, SamplePosition, DynamicObjectSettings.SurfaceLightSampleSpacing))
 						{
 							const FLightRay Ray(
 								SamplePosition,
-								SamplePosition - FVector4(0,0,VolumeBounds.BoxExtent.Z * 2.0f),
+								SamplePosition - FVector4f(0,0,VolumeBounds.BoxExtent.Z * 2.0f),
 								nullptr,
 								nullptr
 								);
@@ -429,7 +429,7 @@ void FStaticLightingSystem::BeginCalculateVolumeSamples()
 							}
 
 							// Add a sample and set its radius such that its influence touches a diagonal sample on the 3d grid.
-							VolumeLightingSampleArray->Add(FVolumeLightingSample(FVector4(SamplePosition, DetailVolumeSpacing * FMath::Sqrt(3.0f))));
+							VolumeLightingSampleArray->Add(FVolumeLightingSample(FVector4f(SamplePosition, DetailVolumeSpacing * FMath::Sqrt(3.0f))));
 							VolumeLightingOctree.AddElement(FVolumeSampleProximityElement(VolumeLightingSampleArray->Num() - 1, *VolumeLightingSampleArray));
 							if (DynamicObjectSettings.bVisualizeVolumeLightInterpolation)
 							{
@@ -473,7 +473,7 @@ void FStaticLightingSystem::BeginCalculateVolumeSamples()
 			{
 				for (float SampleZ = VolumeBounds.Origin.Z - VolumeBounds.BoxExtent.Z; SampleZ < VolumeBounds.Origin.Z + VolumeBounds.BoxExtent.Z + EffectiveVolumeSpacing; SampleZ += EffectiveVolumeSpacing)
 				{
-					const FVector4 SamplePosition(SampleX, SampleY, SampleZ);
+					const FVector4f SamplePosition(SampleX, SampleY, SampleZ);
 					// Only place inside the importance volume
 					if (IsPointInImportanceVolume(SamplePosition, EffectiveVolumeSpacing)
 						// Only place a sample if there are no surface lighting samples nearby
@@ -481,7 +481,7 @@ void FStaticLightingSystem::BeginCalculateVolumeSamples()
 					{
 						NumUniformVolumeSamples++;
 						// Add a sample and set its radius such that its influence touches a diagonal sample on the 3d grid.
-						UniformVolumeSamples->Add(FVolumeLightingSample(FVector4(SamplePosition, EffectiveVolumeSpacing * FMath::Sqrt(3.0f))));
+						UniformVolumeSamples->Add(FVolumeLightingSample(FVector4f(SamplePosition, EffectiveVolumeSpacing * FMath::Sqrt(3.0f))));
 						VolumeLightingOctree.AddElement(FVolumeSampleProximityElement(UniformVolumeSamples->Num() - 1, *UniformVolumeSamples));
 						if (DynamicObjectSettings.bVisualizeVolumeLightInterpolation)
 						{
@@ -521,8 +521,8 @@ void FStaticLightingSystem::ProcessVolumeSamplesTask(const FVolumeSamplesTaskDes
 	FLMRandomStream RandomStream(0);
 	FStaticLightingMappingContext MappingContext(nullptr, *this);
 
-	TArray<FVector4> UniformHemisphereSamples;
-	TArray<FVector2D> UniformHemisphereSampleUniforms;
+	TArray<FVector4f> UniformHemisphereSamples;
+	TArray<FVector2f> UniformHemisphereSampleUniforms;
 	const int32 NumUpperVolumeSamples = ImportanceTracingSettings.NumHemisphereSamples * DynamicObjectSettings.NumHemisphereSamplesScale;
 	// Volume samples don't do any importance sampling so they need more samples for the same amount of variance as surface samples
 	const float NumThetaStepsFloat = FMath::Sqrt(NumUpperVolumeSamples / (float)PI);
@@ -531,7 +531,7 @@ void FStaticLightingSystem::ProcessVolumeSamplesTask(const FVolumeSamplesTaskDes
 
 	GenerateStratifiedUniformHemisphereSamples(NumThetaSteps, NumPhiSteps, RandomStream, UniformHemisphereSamples, UniformHemisphereSampleUniforms);
 
-	FVector4 CombinedVector(0);
+	FVector4f CombinedVector(0);
 
 	for (int32 SampleIndex = 0; SampleIndex < UniformHemisphereSamples.Num(); SampleIndex++)
 	{
@@ -559,8 +559,8 @@ void FStaticLightingSystem::ProcessVolumeSamplesTask(const FVolumeSamplesTaskDes
 			const float SampleRadius = CurrentSample.PositionAndRadius.W;
 			CurrentSample.PositionAndRadius.W = 0.0f;
 
-			TArray<FVector, TInlineAllocator<1>> VertexOffsets;
-			VertexOffsets.Add(FVector(0, 0, 0));
+			TArray<FVector3f, TInlineAllocator<1>> VertexOffsets;
+			VertexOffsets.Add(FVector3f(0, 0, 0));
 
 			CalculateVolumeSampleIncidentRadiance(UniformHemisphereSamples, UniformHemisphereSampleUniforms, MaxUnoccludedLength, VertexOffsets, CurrentSample, BackfacingHitsFraction, Unused, RandomStream, MappingContext, bDebugSamples);
 			
@@ -594,7 +594,7 @@ FGatheredLightSample FStaticLightingSystem::InterpolatePrecomputedVolumeIncident
 	}
 
 	// Iterate over the octree nodes containing the query point.
-	for (FVolumeLightingInterpolationOctree::TConstElementBoxIterator<> OctreeIt(VolumeLightingInterpolationOctree, FBoxCenterAndExtent(Vertex.WorldPosition, FVector4(0,0,0)));
+	for (FVolumeLightingInterpolationOctree::TConstElementBoxIterator<> OctreeIt(VolumeLightingInterpolationOctree, FBoxCenterAndExtent(Vertex.WorldPosition, FVector4f(0,0,0)));
 		OctreeIt.HasPendingElements();
 		OctreeIt.Advance())
 	{
diff --git a/Engine/Source/Programs/UnrealLightmass/Private/Lighting/StaticMesh.cpp b/Engine/Source/Programs/UnrealLightmass/Private/Lighting/StaticMesh.cpp
--- a/Engine/Source/Programs/UnrealLightmass/Private/Lighting/StaticMesh.cpp
+++ b/Engine/Source/Programs/UnrealLightmass/Private/Lighting/StaticMesh.cpp
@@ -25,7 +25,7 @@ static float SmoothStep(float A, float B, float X)
 	return InterpFraction * InterpFraction * (3.0f - 2.0f * InterpFraction);
 }
 
-static FVector SplineEvalPos(const FVector& StartPos, const FVector& StartTangent, const FVector& EndPos, const FVector& EndTangent, float A)
+static FVector3f SplineEvalPos(const FVector3f& StartPos, const FVector3f& StartTangent, const FVector3f& EndPos, const FVector3f& EndTangent, float A)
 {
 	const float A2 = A  * A;
 	const float A3 = A2 * A;
@@ -33,11 +33,11 @@ static FVector SplineEvalPos(const FVector& StartPos, const FVector& StartTangen
 	return (((2*A3)-(3*A2)+1) * StartPos) + ((A3-(2*A2)+A) * StartTangent) + ((A3-A2) * EndTangent) + (((-2*A3)+(3*A2)) * EndPos);
 }
 
-static FVector SplineEvalDir(const FVector& StartPos, const FVector& StartTangent, const FVector& EndPos, const FVector& EndTangent, float A)
+static FVector3f SplineEvalDir(const FVector3f& StartPos, const FVector3f& StartTangent, const FVector3f& EndPos, const FVector3f& EndTangent, float A)
 {
-	const FVector C = (6*StartPos) + (3*StartTangent) + (3*EndTangent) - (6*EndPos);
-	const FVector D = (-6*StartPos) - (4*StartTangent) - (2*EndTangent) + (6*EndPos);
-	const FVector E = StartTangent;
+	const FVector3f C = (6*StartPos) + (3*StartTangent) + (3*EndTangent) - (6*EndPos);
+	const FVector3f D = (-6*StartPos) - (4*StartTangent) - (2*EndTangent) + (6*EndPos);
+	const FVector3f E = StartTangent;
 
 	const float A2 = A  * A;
 
@@ -46,7 +46,7 @@ static FVector SplineEvalDir(const FVector& StartPos, const FVector& StartTangen
 
 /** Calculate full transform that defines frame along spline, given the Z of a vertex. */
 /** Note:  This is mirrored from USplineMeshComponent::CalcSliceTransform() and LocalVertexShader.usf.  If you update one of these, please update them all! */
-static FMatrix CalcSliceTransform(float ZPos, const FSplineMeshParams& SplineParams)
+static FMatrix44f CalcSliceTransform(float ZPos, const FSplineMeshParams& SplineParams)
 {
 	// Find how far 'along' mesh we are
 	const float Alpha = (ZPos - SplineParams.MeshMinZ) / SplineParams.MeshRangeZ;
@@ -55,15 +55,15 @@ static FMatrix CalcSliceTransform(float ZPos, const FSplineMeshParams& SplinePar
 	const float HermiteAlpha = SplineParams.bSmoothInterpRollScale ? SmoothStep(0.0, 1.0, Alpha) : Alpha;
 
 	// Then find the point and direction of the spline at this point along
-	FVector SplinePos = SplineEvalPos( SplineParams.StartPos, SplineParams.StartTangent, SplineParams.EndPos, SplineParams.EndTangent, Alpha );
-	const FVector SplineDir = SplineEvalDir( SplineParams.StartPos, SplineParams.StartTangent, SplineParams.EndPos, SplineParams.EndTangent, Alpha );
+	FVector3f SplinePos = SplineEvalPos( SplineParams.StartPos, SplineParams.StartTangent, SplineParams.EndPos, SplineParams.EndTangent, Alpha );
+	const FVector3f SplineDir = SplineEvalDir( SplineParams.StartPos, SplineParams.StartTangent, SplineParams.EndPos, SplineParams.EndTangent, Alpha );
 
 	// Find base frenet frame
-	const FVector BaseXVec = (SplineParams.SplineUpDir ^ SplineDir).GetSafeNormal();
-	const FVector BaseYVec = (SplineDir ^ BaseXVec).GetSafeNormal();
+	const FVector3f BaseXVec = (SplineParams.SplineUpDir ^ SplineDir).GetSafeNormal();
+	const FVector3f BaseYVec = (SplineDir ^ BaseXVec).GetSafeNormal();
 
 	// Offset the spline by the desired amount
-	const FVector2D SliceOffset = FMath::Lerp<FVector2D>(SplineParams.StartOffset, SplineParams.EndOffset, HermiteAlpha);
+	const FVector2f SliceOffset = FMath::Lerp<FVector2f>(SplineParams.StartOffset, SplineParams.EndOffset, HermiteAlpha);
 	SplinePos += SliceOffset.X * BaseXVec;
 	SplinePos += SliceOffset.Y * BaseYVec;
 
@@ -71,24 +71,24 @@ static FMatrix CalcSliceTransform(float ZPos, const FSplineMeshParams& SplinePar
 	const float UseRoll = FMath::Lerp(SplineParams.StartRoll, SplineParams.EndRoll, HermiteAlpha);
 	const float CosAng = FMath::Cos(UseRoll);
 	const float SinAng = FMath::Sin(UseRoll);
-	const FVector XVec = (CosAng * BaseXVec) - (SinAng * BaseYVec);
-	const FVector YVec = (CosAng * BaseYVec) + (SinAng * BaseXVec);
+	const FVector3f XVec = (CosAng * BaseXVec) - (SinAng * BaseYVec);
+	const FVector3f YVec = (CosAng * BaseYVec) + (SinAng * BaseXVec);
 
 	// Find scale at this point along spline
-	const FVector2D UseScale = FMath::Lerp(SplineParams.StartScale, SplineParams.EndScale, HermiteAlpha);
+	const FVector2f UseScale = FMath::Lerp(SplineParams.StartScale, SplineParams.EndScale, HermiteAlpha);
 
 	// Build overall transform
-	FMatrix SliceTransform;
+	FMatrix44f SliceTransform;
 	switch (SplineParams.ForwardAxis)
 	{
 	case ESplineMeshAxis::X:
-		SliceTransform = FMatrix(SplineDir, UseScale.X * XVec, UseScale.Y * YVec, SplinePos);
+		SliceTransform = FMatrix44f(SplineDir, UseScale.X * XVec, UseScale.Y * YVec, SplinePos);
 		break;
 	case ESplineMeshAxis::Y:
-		SliceTransform = FMatrix(UseScale.Y * YVec, SplineDir, UseScale.X * XVec, SplinePos);
+		SliceTransform = FMatrix44f(UseScale.Y * YVec, SplineDir, UseScale.X * XVec, SplinePos);
 		break;
 	case ESplineMeshAxis::Z:
-		SliceTransform = FMatrix(UseScale.X * XVec, UseScale.Y * YVec, SplineDir, SplinePos);
+		SliceTransform = FMatrix44f(UseScale.X * XVec, UseScale.Y * YVec, SplineDir, SplinePos);
 		break;
 	default:
 		check(0);
@@ -99,7 +99,7 @@ static FMatrix CalcSliceTransform(float ZPos, const FSplineMeshParams& SplinePar
 }
 
 /** Calculate rotation matrix that defines frame along spline, given the Z of a vertex. */
-static FMatrix CalcSliceRot(float ZPos, const FSplineMeshParams& SplineParams)
+static FMatrix44f CalcSliceRot(float ZPos, const FSplineMeshParams& SplineParams)
 {
 	// Find how far 'along' mesh we are
 	const float Alpha = (ZPos - SplineParams.MeshMinZ) / SplineParams.MeshRangeZ;
@@ -108,31 +108,31 @@ static FMatrix CalcSliceRot(float ZPos, const FSplineMeshParams& SplineParams)
 	const float HermiteAlpha = SplineParams.bSmoothInterpRollScale ? SmoothStep(0.0, 1.0, Alpha) : Alpha;
 
 	// Then find the point and direction of the spline at this point along
-	const FVector SplineDir = SplineEvalDir( SplineParams.StartPos, SplineParams.StartTangent, SplineParams.EndPos, SplineParams.EndTangent, Alpha );
+	const FVector3f SplineDir = SplineEvalDir( SplineParams.StartPos, SplineParams.StartTangent, SplineParams.EndPos, SplineParams.EndTangent, Alpha );
 
 	// Find base frenet frame
-	const FVector BaseXVec = (SplineParams.SplineUpDir ^ SplineDir).GetSafeNormal();
-	const FVector BaseYVec = (SplineDir ^ BaseXVec).GetSafeNormal();
+	const FVector3f BaseXVec = (SplineParams.SplineUpDir ^ SplineDir).GetSafeNormal();
+	const FVector3f BaseYVec = (SplineDir ^ BaseXVec).GetSafeNormal();
 
 	// Apply roll to frame around spline
 	const float UseRoll = FMath::Lerp(SplineParams.StartRoll, SplineParams.EndRoll, HermiteAlpha);
 	const float CosAng = FMath::Cos(UseRoll);
 	const float SinAng = FMath::Sin(UseRoll);
-	const FVector XVec = (CosAng * BaseXVec) - (SinAng * BaseYVec);
-	const FVector YVec = (CosAng * BaseYVec) + (SinAng * BaseXVec);
+	const FVector3f XVec = (CosAng * BaseXVec) - (SinAng * BaseYVec);
+	const FVector3f YVec = (CosAng * BaseYVec) + (SinAng * BaseXVec);
 
 	// Build rotation transform
-	FMatrix SliceTransform;
+	FMatrix44f SliceTransform;
 	switch (SplineParams.ForwardAxis)
 	{
 	case ESplineMeshAxis::X:
-		SliceTransform = FMatrix(SplineDir, XVec, YVec, FVector(0,0,0));
+		SliceTransform = FMatrix44f(SplineDir, XVec, YVec, FVector3f(0,0,0));
 		break;
 	case ESplineMeshAxis::Y:
-		SliceTransform = FMatrix(YVec, SplineDir, XVec, FVector(0,0,0));
+		SliceTransform = FMatrix44f(YVec, SplineDir, XVec, FVector3f(0,0,0));
 		break;
 	case ESplineMeshAxis::Z:
-		SliceTransform = FMatrix(XVec, YVec, SplineDir, FVector(0,0,0));
+		SliceTransform = FMatrix44f(XVec, YVec, SplineDir, FVector3f(0,0,0));
 		break;
 	default:
 		check(0);
@@ -150,8 +150,8 @@ static FMatrix CalcSliceRot(float ZPos, const FSplineMeshParams& SplineParams)
  */
 static void GetStaticLightingVertex(
 	const FStaticMeshVertex& InVertex,
-	const FMatrix& LocalToWorld,
-	const FMatrix& LocalToWorldInverseTranspose,
+	const FMatrix44f& LocalToWorld,
+	const FMatrix44f& LocalToWorldInverseTranspose,
 	bool bIsSplineMesh,
 	const FSplineMeshParams& SplineParams,
 	FStaticLightingVertex& OutVertex
@@ -160,23 +160,23 @@ static void GetStaticLightingVertex(
 	if (bIsSplineMesh)
 	{
 		// Make transform for this point along spline
-		const FMatrix SliceTransform = CalcSliceTransform(InVertex.Position[SplineParams.ForwardAxis], SplineParams);
+		const FMatrix44f SliceTransform = CalcSliceTransform(InVertex.Position[SplineParams.ForwardAxis], SplineParams);
 
 		// Remove Z (transform will move us along spline)
-		FVector4 SlicePos = InVertex.Position;
+		FVector4f SlicePos = InVertex.Position;
 		SlicePos[SplineParams.ForwardAxis] = 0;
 
 		// Transform into mesh space
-		const FVector4 LocalPos = SliceTransform.TransformPosition(SlicePos);
+		const FVector4f LocalPos = SliceTransform.TransformPosition(SlicePos);
 
 		// Transform from mesh to world space
 		OutVertex.WorldPosition = LocalToWorld.TransformPosition(LocalPos);
 
-		const FMatrix SliceRot = CalcSliceRot(InVertex.Position[SplineParams.ForwardAxis], SplineParams);
+		const FMatrix44f SliceRot = CalcSliceRot(InVertex.Position[SplineParams.ForwardAxis], SplineParams);
 
-		const FVector4 LocalSpaceTangentX = SliceRot.TransformVector(InVertex.TangentX);
-		const FVector4 LocalSpaceTangentY = SliceRot.TransformVector(InVertex.TangentY);
-		const FVector4 LocalSpaceTangentZ = SliceRot.TransformVector(InVertex.TangentZ);
+		const FVector4f LocalSpaceTangentX = SliceRot.TransformVector(InVertex.TangentX);
+		const FVector4f LocalSpaceTangentY = SliceRot.TransformVector(InVertex.TangentY);
+		const FVector4f LocalSpaceTangentZ = SliceRot.TransformVector(InVertex.TangentZ);
 
 		OutVertex.WorldTangentX = LocalToWorld.TransformVector(LocalSpaceTangentX).GetSafeNormal();
 		OutVertex.WorldTangentY = LocalToWorld.TransformVector(LocalSpaceTangentY).GetSafeNormal();
@@ -244,9 +244,9 @@ void FStaticMeshStaticLightingMesh::GetNonTransformedTriangle(int32 TriangleInde
 	const uint32 I2 = LODRenderData.GetIndex(TriangleIndex * 3 + 2);
 
 	// Translate the triangle's static mesh vertices to static lighting vertices.
-	GetStaticLightingVertex(LODRenderData.GetVertex(I0), FMatrix::Identity, FMatrix::Identity, false, SplineParameters, OutV0);
-	GetStaticLightingVertex(LODRenderData.GetVertex(I1), FMatrix::Identity, FMatrix::Identity, false, SplineParameters, OutV1);
-	GetStaticLightingVertex(LODRenderData.GetVertex(I2), FMatrix::Identity, FMatrix::Identity, false, SplineParameters, OutV2);
+	GetStaticLightingVertex(LODRenderData.GetVertex(I0), FMatrix44f::Identity, FMatrix44f::Identity, false, SplineParameters, OutV0);
+	GetStaticLightingVertex(LODRenderData.GetVertex(I1), FMatrix44f::Identity, FMatrix44f::Identity, false, SplineParameters, OutV1);
+	GetStaticLightingVertex(LODRenderData.GetVertex(I2), FMatrix44f::Identity, FMatrix44f::Identity, false, SplineParameters, OutV2);
 
 	const FStaticMeshLOD& MeshLOD = StaticMesh->GetLOD(GetMeshLODIndex());
 	ElementIndex = INDEX_NONE;
diff --git a/Engine/Source/Programs/UnrealLightmass/Private/Lighting/StaticMesh.h b/Engine/Source/Programs/UnrealLightmass/Private/Lighting/StaticMesh.h
--- a/Engine/Source/Programs/UnrealLightmass/Private/Lighting/StaticMesh.h
+++ b/Engine/Source/Programs/UnrealLightmass/Private/Lighting/StaticMesh.h
@@ -20,7 +20,7 @@ namespace Lightmass
 		FStaticLightingMapping* Mapping;
 
 		/** World to local transform, for box intersection w/ instances. */
-		FMatrix WorldToLocal;
+		FMatrix44f WorldToLocal;
 
 		virtual FStaticMeshStaticLightingMesh* GetInstanceableStaticMesh() override { return !bIsSplineMesh ? this : nullptr; }
 
@@ -73,7 +73,7 @@ namespace Lightmass
 		virtual void Import( class FLightmassImporter& Importer );
 
 		/** The inverse transpose of the primitive's local to world transform. */
-		FMatrix LocalToWorldInverseTranspose;
+		FMatrix44f LocalToWorldInverseTranspose;
 	};
 
 	/** Represents a static mesh primitive with texture mapped static lighting. */
diff --git a/Engine/Source/Programs/UnrealLightmass/Private/Lighting/Texture.h b/Engine/Source/Programs/UnrealLightmass/Private/Lighting/Texture.h
--- a/Engine/Source/Programs/UnrealLightmass/Private/Lighting/Texture.h
+++ b/Engine/Source/Programs/UnrealLightmass/Private/Lighting/Texture.h
@@ -78,7 +78,7 @@ namespace Lightmass
 			FMemory::Memzero(Data, ElementSize * SizeX * SizeY);
 		}
 
-		inline uint8* SampleRawPtr(const FVector2D& UV) const
+		inline uint8* SampleRawPtr(const FVector2f& UV) const
 		{
 			// Wrapped addressing (uses FMath::FloorToInt and not appFractional, as appFractional causes the
 			// following mapping:
@@ -99,7 +99,7 @@ namespace Lightmass
 			return &Data[DataIndex];
 		}
 
-		inline FLinearColor Sample(const FVector2D& UV) const
+		inline FLinearColor Sample(const FVector2f& UV) const
 		{
 			uint8* RawPtr = SampleRawPtr( UV );
 
@@ -115,7 +115,7 @@ namespace Lightmass
 			return FLinearColor(*(FColor*)RawPtr);
 		}
 
-		inline FVector4 SampleNormal(const FVector2D& UV) const
+		inline FVector4f SampleNormal(const FVector2f& UV) const
 		{
 			uint8* RawPtr = SampleRawPtr( UV );
 
@@ -123,7 +123,7 @@ namespace Lightmass
 			checkSlow(Format == TF_ARGB16F);
 			FFloat16Color* Float16Color = (FFloat16Color*)RawPtr;
 
-			FVector4 Normal( Float16Color->R.GetFloat(), Float16Color->G.GetFloat(), Float16Color->B.GetFloat(), 0.0f );
+			FVector4f Normal( Float16Color->R.GetFloat(), Float16Color->G.GetFloat(), Float16Color->B.GetFloat(), 0.0f );
 			return Normal;
 		}
 	};
