Fix CPU Lightmass by coverting LWC types to single precision

diff --git a/Engine/Source/Programs/UnrealLightmass/Private/Lighting/TextureMapping.cpp b/Engine/Source/Programs/UnrealLightmass/Private/Lighting/TextureMapping.cpp
--- a/Engine/Source/Programs/UnrealLightmass/Private/Lighting/TextureMapping.cpp
+++ b/Engine/Source/Programs/UnrealLightmass/Private/Lighting/TextureMapping.cpp
@@ -62,9 +62,9 @@ void FStaticLightingRasterPolicy::ProcessPixel(int32 X,int32 Y,const Interpolant
 
 		// Add this sample to the mapping.
 		TexelToVertex.WorldPosition = TexelToVertex.WorldPosition * OldSampleWeight + Interpolant.Vertex.WorldPosition * NewSampleWeight;
-		TexelToVertex.WorldTangentX = FVector4(TexelToVertex.WorldTangentX) * OldSampleWeight + Interpolant.Vertex.WorldTangentX * NewSampleWeight;
-		TexelToVertex.WorldTangentY = FVector4(TexelToVertex.WorldTangentY) * OldSampleWeight + Interpolant.Vertex.WorldTangentY * NewSampleWeight;
-		TexelToVertex.WorldTangentZ = FVector4(TexelToVertex.WorldTangentZ) * OldSampleWeight + Interpolant.Vertex.WorldTangentZ * NewSampleWeight;
+		TexelToVertex.WorldTangentX = FVector4f(TexelToVertex.WorldTangentX) * OldSampleWeight + Interpolant.Vertex.WorldTangentX * NewSampleWeight;
+		TexelToVertex.WorldTangentY = FVector4f(TexelToVertex.WorldTangentY) * OldSampleWeight + Interpolant.Vertex.WorldTangentY * NewSampleWeight;
+		TexelToVertex.WorldTangentZ = FVector4f(TexelToVertex.WorldTangentZ) * OldSampleWeight + Interpolant.Vertex.WorldTangentZ * NewSampleWeight;
 		TexelToVertex.TriangleNormal = TriangleNormal;
 		TexelToVertex.ElementIndex = Interpolant.ElementIndex;
 		
@@ -79,12 +79,12 @@ void FFullStaticLightingVertex::ApplyVertexModifications(int32 ElementIndex, boo
 {
 	if (bUseNormalMapsForLighting && Mesh->HasImportedNormal(ElementIndex))
 	{
-		const FVector4 TangentNormal = Mesh->EvaluateNormal(TextureCoordinates[0], ElementIndex);
+		const FVector4f TangentNormal = Mesh->EvaluateNormal(TextureCoordinates[0], ElementIndex);
 
-		const FVector4 WorldTangentRow0(WorldTangentX.X, WorldTangentY.X, WorldTangentZ.X);
-		const FVector4 WorldTangentRow1(WorldTangentX.Y, WorldTangentY.Y, WorldTangentZ.Y);
-		const FVector4 WorldTangentRow2(WorldTangentX.Z, WorldTangentY.Z, WorldTangentZ.Z);
-		const FVector4 WorldVector(
+		const FVector4f WorldTangentRow0(WorldTangentX.X, WorldTangentY.X, WorldTangentZ.X);
+		const FVector4f WorldTangentRow1(WorldTangentX.Y, WorldTangentY.Y, WorldTangentZ.Y);
+		const FVector4f WorldTangentRow2(WorldTangentX.Z, WorldTangentY.Z, WorldTangentZ.Z);
+		const FVector4f WorldVector(
 			Dot3(WorldTangentRow0, TangentNormal),
 			Dot3(WorldTangentRow1, TangentNormal),
 			Dot3(WorldTangentRow2, TangentNormal)
@@ -100,8 +100,8 @@ void FFullStaticLightingVertex::ApplyVertexModifications(int32 ElementIndex, boo
 	TriangleNormal = bUseVertexNormalForHemisphereGather ? WorldTangentZ : TriangleNormal.GetUnsafeNormal3();
 	checkSlow(!TriangleNormal.ContainsNaN());
 
-	const FVector4 OriginalTangentX = WorldTangentX;
-	const FVector4 OriginalTangentY = WorldTangentY;
+	const FVector4f OriginalTangentX = WorldTangentX;
+	const FVector4f OriginalTangentY = WorldTangentY;
 
 	WorldTangentY = (WorldTangentZ ^ WorldTangentX).GetUnsafeNormal3();
 	// Maintain handedness
@@ -128,7 +128,7 @@ void FStaticLightingSystem::CacheIrradiancePhotonsTextureMapping(FStaticLighting
 	checkSlow(TextureMapping);
 	FStaticLightingMappingContext MappingContext(TextureMapping->Mesh,*this,&StartupDebugOutput);
 	LIGHTINGSTAT(FScopedRDTSCTimer CachingTime(MappingContext.Stats.IrradiancePhotonCachingThreadTime));
-	const FBoxSphereBounds ImportanceBounds = Scene.GetImportanceBounds();
+	const FBoxSphereBounds3f ImportanceBounds = Scene.GetImportanceBounds();
 
 	FTexelToVertexMap TexelToVertexMap(TextureMapping->SurfaceCacheSizeX, TextureMapping->SurfaceCacheSizeY);
 
@@ -282,14 +282,14 @@ void FStaticLightingSystem::RasterizeToSurfaceCacheTextureMapping(FStaticLightin
 		TextureMapping->Mesh->GetTriangle(TriangleIndex,V0.Vertex,V1.Vertex,V2.Vertex,Element);
 		V0.ElementIndex = V1.ElementIndex = V2.ElementIndex = Element;
 
-		const FVector4 TriangleNormal = ((V2.Vertex.WorldPosition - V0.Vertex.WorldPosition) ^ (V1.Vertex.WorldPosition - V0.Vertex.WorldPosition)).GetSafeNormal();
+		const FVector4f TriangleNormal = ((V2.Vertex.WorldPosition - V0.Vertex.WorldPosition) ^ (V1.Vertex.WorldPosition - V0.Vertex.WorldPosition)).GetSafeNormal();
 
 		// Don't rasterize degenerates 
 		if (!TriangleNormal.IsNearlyZero3())
 		{
-			const FVector2D UV0 = V0.Vertex.TextureCoordinates[TextureMapping->LightmapTextureCoordinateIndex] * FVector2D(TextureMapping->SurfaceCacheSizeX,TextureMapping->SurfaceCacheSizeY);
-			const FVector2D UV1 = V1.Vertex.TextureCoordinates[TextureMapping->LightmapTextureCoordinateIndex] * FVector2D(TextureMapping->SurfaceCacheSizeX,TextureMapping->SurfaceCacheSizeY);
-			const FVector2D UV2 = V2.Vertex.TextureCoordinates[TextureMapping->LightmapTextureCoordinateIndex] * FVector2D(TextureMapping->SurfaceCacheSizeX,TextureMapping->SurfaceCacheSizeY);
+			const FVector2f UV0 = V0.Vertex.TextureCoordinates[TextureMapping->LightmapTextureCoordinateIndex] * FVector2f(TextureMapping->SurfaceCacheSizeX,TextureMapping->SurfaceCacheSizeY);
+			const FVector2f UV1 = V1.Vertex.TextureCoordinates[TextureMapping->LightmapTextureCoordinateIndex] * FVector2f(TextureMapping->SurfaceCacheSizeX,TextureMapping->SurfaceCacheSizeY);
+			const FVector2f UV2 = V2.Vertex.TextureCoordinates[TextureMapping->LightmapTextureCoordinateIndex] * FVector2f(TextureMapping->SurfaceCacheSizeX,TextureMapping->SurfaceCacheSizeY);
 
 			// Odd number of samples so that the center of the pyramid is on one of the samples
 			const uint32 NumSamplesX = 5;
@@ -321,9 +321,9 @@ void FStaticLightingSystem::RasterizeToSurfaceCacheTextureMapping(FStaticLightin
 						V0,
 						V1,
 						V2,
-						UV0 + FVector2D(SampleXOffset, SampleYOffset),
-						UV1 + FVector2D(SampleXOffset, SampleYOffset),
-						UV2 + FVector2D(SampleXOffset, SampleYOffset),
+						UV0 + FVector2f(SampleXOffset, SampleYOffset),
+						UV1 + FVector2f(SampleXOffset, SampleYOffset),
+						UV2 + FVector2f(SampleXOffset, SampleYOffset),
 						false
 						);
 				}
@@ -420,12 +420,12 @@ void FStaticLightingSystem::AdjustRepresentativeSurfelForTexelsTextureMapping(
 
 				if (MaterialSettings.bUseNormalMapsForLighting && TextureMapping->Mesh->HasImportedNormal(TexelToVertex.ElementIndex))
 				{
-					const FVector4 TangentNormal = TextureMapping->Mesh->EvaluateNormal(TexelToVertex.TextureCoordinates[0], TexelToVertex.ElementIndex);
+					const FVector4f TangentNormal = TextureMapping->Mesh->EvaluateNormal(TexelToVertex.TextureCoordinates[0], TexelToVertex.ElementIndex);
 
-					const FVector4 WorldTangentRow0(TexelToVertex.WorldTangentX.X, TexelToVertex.WorldTangentY.X, TexelToVertex.WorldTangentZ.X);
-					const FVector4 WorldTangentRow1(TexelToVertex.WorldTangentX.Y, TexelToVertex.WorldTangentY.Y, TexelToVertex.WorldTangentZ.Y);
-					const FVector4 WorldTangentRow2(TexelToVertex.WorldTangentX.Z, TexelToVertex.WorldTangentY.Z, TexelToVertex.WorldTangentZ.Z);
-					const FVector4 WorldVector(
+					const FVector4f WorldTangentRow0(TexelToVertex.WorldTangentX.X, TexelToVertex.WorldTangentY.X, TexelToVertex.WorldTangentZ.X);
+					const FVector4f WorldTangentRow1(TexelToVertex.WorldTangentX.Y, TexelToVertex.WorldTangentY.Y, TexelToVertex.WorldTangentZ.Y);
+					const FVector4f WorldTangentRow2(TexelToVertex.WorldTangentX.Z, TexelToVertex.WorldTangentY.Z, TexelToVertex.WorldTangentZ.Z);
+					const FVector4f WorldVector(
 						Dot3(WorldTangentRow0, TangentNormal),
 						Dot3(WorldTangentRow1, TangentNormal),
 						Dot3(WorldTangentRow2, TangentNormal)
@@ -441,8 +441,8 @@ void FStaticLightingSystem::AdjustRepresentativeSurfelForTexelsTextureMapping(
 				TexelToVertex.TriangleNormal = bUseVertexNormalForHemisphereGather ? TexelToVertex.WorldTangentZ : TexelToVertex.TriangleNormal.GetUnsafeNormal3();
 				checkSlow(!TexelToVertex.TriangleNormal.ContainsNaN());
 
-				const FVector4 OriginalTangentX = TexelToVertex.WorldTangentX;
-				const FVector4 OriginalTangentY = TexelToVertex.WorldTangentY;
+				const FVector4f OriginalTangentX = TexelToVertex.WorldTangentX;
+				const FVector4f OriginalTangentY = TexelToVertex.WorldTangentY;
 
 				TexelToVertex.WorldTangentY = (TexelToVertex.WorldTangentZ ^ TexelToVertex.WorldTangentX).GetUnsafeNormal3();
 				// Maintain handedness
@@ -491,16 +491,16 @@ void FStaticLightingSystem::AdjustRepresentativeSurfelForTexelsTextureMapping(
 
 				{
 					const FFullStaticLightingVertex FullVertex = TexelToVertex.GetFullVertex();
-					const FVector4 TexelCenterOffset = FullVertex.WorldPosition + FullVertex.TriangleNormal * TexelToVertex.TexelRadius * SceneConstants.VisibilityNormalOffsetSampleRadiusScale;
+					const FVector4f TexelCenterOffset = FullVertex.WorldPosition + FullVertex.TriangleNormal * TexelToVertex.TexelRadius * SceneConstants.VisibilityNormalOffsetSampleRadiusScale;
 
 					FLightRayIntersection Intersections[4];
 					bool bHitBackfaces[4];
 
-					FVector2D CornerSigns[4];
-					CornerSigns[0] = FVector2D(1, 1);
-					CornerSigns[1] = FVector2D(-1, 1);
-					CornerSigns[2] = FVector2D(1, -1);
-					CornerSigns[3] = FVector2D(-1, -1);
+					FVector2f CornerSigns[4];
+					CornerSigns[0] = FVector2f(1, 1);
+					CornerSigns[1] = FVector2f(-1, 1);
+					CornerSigns[2] = FVector2f(1, -1);
+					CornerSigns[3] = FVector2f(-1, -1);
 
 					float BackfaceDetectionDistance = TexelToVertex.TexelRadius * 2.5f;
 
@@ -573,7 +573,7 @@ void FStaticLightingSystem::AdjustRepresentativeSurfelForTexelsTextureMapping(
 					if (IntersectionIndexForShadingPositionMovement != INDEX_NONE)
 					{
 						// Move the shading position outside the surface that is intersecting this texel
-						const FVector4 OffsetShadingPosition = Intersections[IntersectionIndexForShadingPositionMovement].IntersectionVertex.WorldPosition
+						const FVector4f OffsetShadingPosition = Intersections[IntersectionIndexForShadingPositionMovement].IntersectionVertex.WorldPosition
 							// Move along the intersecting surface's normal but also away from the texel a bit to prevent incorrect self occlusion
 							+ (Intersections[IntersectionIndexForShadingPositionMovement].IntersectionVertex.WorldTangentZ + TexelToVertex.TriangleNormal) * 0.5f * TexelToVertex.TexelRadius * SceneConstants.VisibilityNormalOffsetSampleRadiusScale;
 
@@ -604,7 +604,7 @@ void FStaticLightingSystem::FinalizeSurfaceCacheTextureMapping(FStaticLightingTe
 #if LIGHTMASS_DO_PROCESSING
 	checkSlow(TextureMapping);
 	FStaticLightingMappingContext MappingContext(TextureMapping->Mesh,*this,&StartupDebugOutput);
-	const FBoxSphereBounds ImportanceBounds = Scene.GetImportanceBounds();
+	const FBoxSphereBounds3f ImportanceBounds = Scene.GetImportanceBounds();
 
 	FTexelToVertexMap TexelToVertexMap(TextureMapping->SurfaceCacheSizeX, TextureMapping->SurfaceCacheSizeY);
 
@@ -667,8 +667,8 @@ void FStaticLightingSystem::FinalizeSurfaceCacheTextureMapping(FStaticLightingTe
 					FGatheredLightSample DirectLighting;
 					FGatheredLightSample Unused;
 					float Unused2;
-					TArray<FVector, TInlineAllocator<1>> VertexOffsets;
-					VertexOffsets.Add(FVector(0, 0, 0));
+					TArray<FVector3f, TInlineAllocator<1>> VertexOffsets;
+					VertexOffsets.Add(FVector3f(0, 0, 0));
 
 					CalculateApproximateDirectLighting(CurrentVertex, TexelToVertex.TexelRadius, VertexOffsets, .1f, true, true, bDebugThisTexel && PhotonMappingSettings.bVisualizeCachedApproximateDirectLighting, MappingContext, DirectLighting, Unused, Unused2);
 					FinalIncidentLighting += DirectLighting.IncidentLighting;
@@ -748,16 +748,16 @@ void FStaticLightingSystem::ProcessTextureMapping(FStaticLightingTextureMapping*
 				if (TexelToVertex.TotalSampleWeight > 0.0f)
 				{
 					// Verify that vertex normals are normalized (within some error that is large because of FPackedNormal)
-					checkSlow(FVector(TexelToVertex.WorldTangentZ).IsUnit(.1f));
+					checkSlow(FVector3f(TexelToVertex.WorldTangentZ).IsUnit(.1f));
 
-					const float DistanceToDebugTexelSq = FVector(TexelToVertex.WorldPosition - Scene.DebugInput.Position).SizeSquared();
+					const float DistanceToDebugTexelSq = FVector3f(TexelToVertex.WorldPosition - Scene.DebugInput.Position).SizeSquared();
 					if (DistanceToDebugTexelSq < 40000
 						|| X == Scene.DebugInput.LocalX && Y == Scene.DebugInput.LocalY)
 					{
 						if (X == Scene.DebugInput.LocalX && Y == Scene.DebugInput.LocalY)
 						{		
 							FDebugStaticLightingVertex DebugVertex;
-							DebugVertex.VertexNormal = FVector4(TexelToVertex.WorldTangentZ);
+							DebugVertex.VertexNormal = FVector4f(TexelToVertex.WorldTangentZ);
 							DebugVertex.VertexPosition = TexelToVertex.WorldPosition;
 							MappingContext.DebugOutput->Vertices.Add(DebugVertex);
 
@@ -791,7 +791,7 @@ void FStaticLightingSystem::ProcessTextureMapping(FStaticLightingTextureMapping*
 					continue;
 				}
 
-				if (!Light->AffectsBounds(FBoxSphereBounds(TextureMapping->Mesh->BoundingBox)))
+				if (!Light->AffectsBounds(FBoxSphereBounds3f(TextureMapping->Mesh->BoundingBox)))
 				{
 					continue;
 				}
@@ -1053,9 +1053,9 @@ void FTexelCornerRasterPolicy::ProcessPixel(int32 X, int32 Y, const InterpolantT
 }
 
 void FStaticLightingSystem::TraceToTexelCorner(
-	const FVector4& TexelCenterOffset, 
+	const FVector4f& TexelCenterOffset, 
 	const FFullStaticLightingVertex& FullVertex, 
-	FVector2D CornerSigns,
+	FVector2f CornerSigns,
 	float TexelRadius, 
 	FStaticLightingMappingContext& MappingContext, 
 	FLightRayIntersection& Intersection,
@@ -1064,7 +1064,7 @@ void FStaticLightingSystem::TraceToTexelCorner(
 {
 	// Vector from the center to one of the corners of the texel
 	// The FMath::Sqrt(.5f) is to normalize (Vertex.TriangleTangentX + Vertex.TriangleTangentY), which are orthogonal unit vectors.
-	const FVector4 CornerOffset = FMath::Sqrt(.5f) * (CornerSigns.X * FullVertex.TriangleTangentX + CornerSigns.Y * FullVertex.TriangleTangentY) * TexelRadius * SceneConstants.VisibilityTangentOffsetSampleRadiusScale;
+	const FVector4f CornerOffset = FMath::Sqrt(.5f) * (CornerSigns.X * FullVertex.TriangleTangentX + CornerSigns.Y * FullVertex.TriangleTangentY) * TexelRadius * SceneConstants.VisibilityTangentOffsetSampleRadiusScale;
 	const FLightRay TexelRay(
 		TexelCenterOffset,
 		TexelCenterOffset + CornerOffset,
@@ -1126,14 +1126,14 @@ void FStaticLightingSystem::SetupTextureMapping(
 			TextureMapping->Mesh->GetTriangle(TriangleIndex,V0.Vertex,V1.Vertex,V2.Vertex,Element);
 			V0.ElementIndex = V1.ElementIndex = V2.ElementIndex = Element;
 
-			const FVector4 TriangleNormal = ((V2.Vertex.WorldPosition - V0.Vertex.WorldPosition) ^ (V1.Vertex.WorldPosition - V0.Vertex.WorldPosition)).GetSafeNormal();
+			const FVector4f TriangleNormal = ((V2.Vertex.WorldPosition - V0.Vertex.WorldPosition) ^ (V1.Vertex.WorldPosition - V0.Vertex.WorldPosition)).GetSafeNormal();
 
 			// Don't rasterize degenerates 
 			if (!TriangleNormal.IsNearlyZero3())
 			{
-				const FVector2D UV0 = V0.Vertex.TextureCoordinates[TextureMapping->LightmapTextureCoordinateIndex] * FVector2D(TextureMapping->CachedSizeX,TextureMapping->CachedSizeY);
-				const FVector2D UV1 = V1.Vertex.TextureCoordinates[TextureMapping->LightmapTextureCoordinateIndex] * FVector2D(TextureMapping->CachedSizeX,TextureMapping->CachedSizeY);
-				const FVector2D UV2 = V2.Vertex.TextureCoordinates[TextureMapping->LightmapTextureCoordinateIndex] * FVector2D(TextureMapping->CachedSizeX,TextureMapping->CachedSizeY);
+				const FVector2f UV0 = V0.Vertex.TextureCoordinates[TextureMapping->LightmapTextureCoordinateIndex] * FVector2f(TextureMapping->CachedSizeX,TextureMapping->CachedSizeY);
+				const FVector2f UV1 = V1.Vertex.TextureCoordinates[TextureMapping->LightmapTextureCoordinateIndex] * FVector2f(TextureMapping->CachedSizeX,TextureMapping->CachedSizeY);
+				const FVector2f UV2 = V2.Vertex.TextureCoordinates[TextureMapping->LightmapTextureCoordinateIndex] * FVector2f(TextureMapping->CachedSizeX,TextureMapping->CachedSizeY);
 
 				// Odd number of samples so that the center of the pyramid is on one of the samples
 				const uint32 NumSamplesX = 7;
@@ -1165,9 +1165,9 @@ void FStaticLightingSystem::SetupTextureMapping(
 							V0,
 							V1,
 							V2,
-							UV0 + FVector2D(SampleXOffset, SampleYOffset),
-							UV1 + FVector2D(SampleXOffset, SampleYOffset),
-							UV2 + FVector2D(SampleXOffset, SampleYOffset),
+							UV0 + FVector2f(SampleXOffset, SampleYOffset),
+							UV1 + FVector2f(SampleXOffset, SampleYOffset),
+							UV2 + FVector2f(SampleXOffset, SampleYOffset),
 							false
 							);
 					}
@@ -1195,7 +1195,7 @@ void FStaticLightingSystem::SetupTextureMapping(
 				Scene,
 				TexelToVertexMap,
 				SampleWeight,
-				FVector4(0),
+				FVector4f(0),
 				bDebugThisMapping,
 				false
 				));
@@ -1205,9 +1205,9 @@ void FStaticLightingSystem::SetupTextureMapping(
 				V0,
 				V1,
 				V2,
-				V0.Vertex.TextureCoordinates[TextureMapping->LightmapTextureCoordinateIndex] * FVector2D(TextureMapping->CachedSizeX,TextureMapping->CachedSizeY) + FVector2D(-0.5f,-0.5f),
-				V1.Vertex.TextureCoordinates[TextureMapping->LightmapTextureCoordinateIndex] * FVector2D(TextureMapping->CachedSizeX,TextureMapping->CachedSizeY) + FVector2D(-0.5f,-0.5f),
-				V2.Vertex.TextureCoordinates[TextureMapping->LightmapTextureCoordinateIndex] * FVector2D(TextureMapping->CachedSizeX,TextureMapping->CachedSizeY) + FVector2D(-0.5f,-0.5f),
+				V0.Vertex.TextureCoordinates[TextureMapping->LightmapTextureCoordinateIndex] * FVector2f(TextureMapping->CachedSizeX,TextureMapping->CachedSizeY) + FVector2f(-0.5f,-0.5f),
+				V1.Vertex.TextureCoordinates[TextureMapping->LightmapTextureCoordinateIndex] * FVector2f(TextureMapping->CachedSizeX,TextureMapping->CachedSizeY) + FVector2f(-0.5f,-0.5f),
+				V2.Vertex.TextureCoordinates[TextureMapping->LightmapTextureCoordinateIndex] * FVector2f(TextureMapping->CachedSizeX,TextureMapping->CachedSizeY) + FVector2f(-0.5f,-0.5f),
 				false
 				);
 		}
@@ -1248,7 +1248,7 @@ void FStaticLightingSystem::CalculateDirectLightingTextureMappingFiltered(
 				ShadowSample.bIsMapped = true;
 
 				// Check if the light is in front of the surface.
-				const bool bLightIsInFrontOfTriangle = !IsLightBehindSurface(TexelToVertex.WorldPosition,FVector4(TexelToVertex.WorldTangentZ),Light);
+				const bool bLightIsInFrontOfTriangle = !IsLightBehindSurface(TexelToVertex.WorldPosition,FVector4f(TexelToVertex.WorldTangentZ),Light);
 				if (bLightIsInFrontOfTriangle || TextureMapping->Mesh->IsTwoSided(TexelToVertex.ElementIndex))
 				{
 					// Compute the shadow factors for this sample from the shadow-mapped lights.
@@ -1442,7 +1442,7 @@ void FStaticLightingSystem::CalculateDirectAreaLightingTextureMapping(
 			const FTexelToVertex& TexelToVertex = TexelToVertexMap(X,Y);
 
 			if (CurrentLightSample.bIsMapped
-				&& Light->AffectsBounds(FBoxSphereBounds(FSphere(TexelToVertex.WorldPosition, TexelToVertex.TexelRadius * 2))))
+				&& Light->AffectsBounds(FBoxSphereBounds3f(FSphere3f(TexelToVertex.WorldPosition, TexelToVertex.TexelRadius * 2))))
 			{
 				UnfilteredShadowFactorData(X, Y).bIsMapped = true;
 
@@ -1462,8 +1462,8 @@ void FStaticLightingSystem::CalculateDirectAreaLightingTextureMapping(
 				// by a reasonably small threshold, and then make sure the lighting code handles rays
 				// that aren't necessarily in front of the triangle robustly.
 				//
-				//		const FVector4 Normal = Vertex.TransformTangentVectorToWorld(TextureMapping->Mesh->EvaluateNormal(Vertex.TextureCoordinates[0], TexelToVertex.ElementIndex)) :*/
-				const FVector4 Normal = Vertex.WorldTangentZ;
+				//		const FVector4f Normal = Vertex.TransformTangentVectorToWorld(TextureMapping->Mesh->EvaluateNormal(Vertex.TextureCoordinates[0], TexelToVertex.ElementIndex)) :*/
+				const FVector4f Normal = Vertex.WorldTangentZ;
 
 				const bool bLightIsInFrontOfTriangle = !Light->BehindSurface(TexelToVertex.WorldPosition, Normal);
 				if (bLightIsInFrontOfTriangle || TextureMapping->Mesh->IsTwoSided(TexelToVertex.ElementIndex))
@@ -1497,7 +1497,7 @@ void FStaticLightingSystem::CalculateDirectAreaLightingTextureMapping(
 						const TArray<FLightSurfaceSample>& LightSurfaceSamples = Light->GetCachedSurfaceSamples(0, false);
 						FLinearColor UnnormalizedTransmission;
 
-						const FVector2D ShadowValue = CalculatePointAreaShadowing(
+						const FVector2f ShadowValue = CalculatePointAreaShadowing(
 							TextureMapping, 
 							CurrentVertex, 
 							TexelToVertex.ElementIndex,
@@ -1517,7 +1517,7 @@ void FStaticLightingSystem::CalculateDirectAreaLightingTextureMapping(
 								const TArray<FLightSurfaceSample>& PenumbraLightSurfaceSamples = Light->GetCachedSurfaceSamples(0, true);
 								FLinearColor UnnormalizedPenumbraTransmission;
 
-								const FVector2D ShadowValuePenumbra = CalculatePointAreaShadowing(
+								const FVector2f ShadowValuePenumbra = CalculatePointAreaShadowing(
 									TextureMapping, 
 									CurrentVertex, 
 									TexelToVertex.ElementIndex,
@@ -1761,7 +1761,7 @@ class FVisibilitySample
 {
 protected:
 	/** World space position in XYZ, Distance to the nearest occluder in W, only valid if !bVisible. */
-	FVector4 PositionAndOccluderDistance;
+	FVector4f PositionAndOccluderDistance;
 	/** World space normal */
 	float NormalX, NormalY, NormalZ;
 	/** Whether this sample is visible to the light. */
@@ -1772,14 +1772,14 @@ protected:
 	uint32 bNeedsHighResSampling : 1;
 
 public:
-	inline FVector4 GetPosition() const { return FVector4(PositionAndOccluderDistance, 0.0f); }
+	inline FVector4f GetPosition() const { return FVector4f(PositionAndOccluderDistance, 0.0f); }
 	inline float GetOccluderDistance() const { return PositionAndOccluderDistance.W; }
-	inline FVector4 GetNormal() const { return FVector4(NormalX, NormalY, NormalZ); }
+	inline FVector4f GetNormal() const { return FVector4f(NormalX, NormalY, NormalZ); }
 	inline bool IsVisible() const { return bVisible; }
 	inline bool IsMapped() const { return bIsMapped; }
 	inline bool NeedsHighResSampling() const { return bNeedsHighResSampling; }
 
-	inline void SetPosition(const FVector4& InPosition)  
+	inline void SetPosition(const FVector4f& InPosition)  
 	{  
 		PositionAndOccluderDistance.X = InPosition.X;
 		PositionAndOccluderDistance.Y = InPosition.Y;
@@ -1789,7 +1789,7 @@ public:
 	{  
 		PositionAndOccluderDistance.W = InOccluderDistance;
 	}
-	inline void SetNormal(const FVector4& InNormal)  
+	inline void SetNormal(const FVector4f& InNormal)  
 	{  
 		NormalX = InNormal.X;
 		NormalY = InNormal.Y;
@@ -1917,15 +1917,15 @@ void FStaticLightingSystem::CalculateDirectSignedDistanceFieldLightingTextureMap
 		TextureMapping->Mesh->GetTriangle(TriangleIndex, MeshVertices[TriangleIndex * 3].Vertex, MeshVertices[TriangleIndex * 3 + 1].Vertex, MeshVertices[TriangleIndex * 3 + 2].Vertex, Element);
 		MeshVertices[TriangleIndex * 3].ElementIndex = MeshVertices[TriangleIndex * 3 + 1].ElementIndex = MeshVertices[TriangleIndex * 3 + 2].ElementIndex = Element;
 
-		const FVector4 TriangleNormal = (MeshVertices[TriangleIndex * 3 + 2].Vertex.WorldPosition - MeshVertices[TriangleIndex * 3].Vertex.WorldPosition) ^ (MeshVertices[TriangleIndex * 3 + 1].Vertex.WorldPosition - MeshVertices[TriangleIndex].Vertex.WorldPosition);
+		const FVector4f TriangleNormal = (MeshVertices[TriangleIndex * 3 + 2].Vertex.WorldPosition - MeshVertices[TriangleIndex * 3].Vertex.WorldPosition) ^ (MeshVertices[TriangleIndex * 3 + 1].Vertex.WorldPosition - MeshVertices[TriangleIndex].Vertex.WorldPosition);
 		const float TriangleArea = 0.5f * TriangleNormal.Size3();
 
 		if (TriangleArea > DELTA)
 		{
 			// Triangle vertices in lightmap UV space, scaled by the lightmap resolution
-			const FVector2D Vertex0 = MeshVertices[TriangleIndex * 3 + 0].Vertex.TextureCoordinates[TextureMapping->LightmapTextureCoordinateIndex] * FVector2D(TextureMapping->CachedSizeX, TextureMapping->CachedSizeY);
-			const FVector2D Vertex1 = MeshVertices[TriangleIndex * 3 + 1].Vertex.TextureCoordinates[TextureMapping->LightmapTextureCoordinateIndex] * FVector2D(TextureMapping->CachedSizeX, TextureMapping->CachedSizeY);
-			const FVector2D Vertex2 = MeshVertices[TriangleIndex * 3 + 2].Vertex.TextureCoordinates[TextureMapping->LightmapTextureCoordinateIndex] * FVector2D(TextureMapping->CachedSizeX, TextureMapping->CachedSizeY);
+			const FVector2f Vertex0 = MeshVertices[TriangleIndex * 3 + 0].Vertex.TextureCoordinates[TextureMapping->LightmapTextureCoordinateIndex] * FVector2f(TextureMapping->CachedSizeX, TextureMapping->CachedSizeY);
+			const FVector2f Vertex1 = MeshVertices[TriangleIndex * 3 + 1].Vertex.TextureCoordinates[TextureMapping->LightmapTextureCoordinateIndex] * FVector2f(TextureMapping->CachedSizeX, TextureMapping->CachedSizeY);
+			const FVector2f Vertex2 = MeshVertices[TriangleIndex * 3 + 2].Vertex.TextureCoordinates[TextureMapping->LightmapTextureCoordinateIndex] * FVector2f(TextureMapping->CachedSizeX, TextureMapping->CachedSizeY);
 
 			// Area in lightmap space, or the number of lightmap texels covered by this triangle
 			const float LightmapTriangleArea = FMath::Abs(
@@ -1979,7 +1979,7 @@ void FStaticLightingSystem::CalculateDirectSignedDistanceFieldLightingTextureMap
 			{
 				NumMappedTexels++;
 				// Note: not checking for backfacing normals because some of the high resolution samples corresponding to this texel may be frontfacing
-				if (Light->AffectsBounds(FBoxSphereBounds(TexelToVertex.WorldPosition, FVector4(0,0,0),0)))
+				if (Light->AffectsBounds(FBoxSphereBounds3f(TexelToVertex.WorldPosition, FVector4f(0,0,0),0)))
 				{
 					FLowResolutionVisibilitySample& CurrentSample = LowResolutionVisibilityData(X, Y);
 					CurrentSample.SetPosition(TexelToVertex.WorldPosition);
@@ -1989,10 +1989,10 @@ void FStaticLightingSystem::CalculateDirectSignedDistanceFieldLightingTextureMap
 					// And multiple shadowmaps on the same object may be merged together, but only if each one marks the area that it has valid data
 					CurrentSample.SetMapped(true);
 
-					const FVector4 LightPosition = Light->LightCenterPosition(TexelToVertex.WorldPosition, TexelToVertex.WorldTangentZ);
-					const FVector4 LightVector = (LightPosition - TexelToVertex.WorldPosition).GetSafeNormal();
+					const FVector4f LightPosition = Light->LightCenterPosition(TexelToVertex.WorldPosition, TexelToVertex.WorldTangentZ);
+					const FVector4f LightVector = (LightPosition - TexelToVertex.WorldPosition).GetSafeNormal();
 
-					FVector4 NormalForOffset = CurrentSample.GetNormal();
+					FVector4f NormalForOffset = CurrentSample.GetNormal();
 					// Flip the normal used for offsetting the start of the ray for two sided materials if a flipped normal would be closer to the light.
 					// This prevents incorrect shadowing where using the frontface normal would cause the ray to start inside a nearby object.
 					const bool bIsTwoSided = TextureMapping->Mesh->IsTwoSided(CurrentSample.ElementIndex);
@@ -2133,9 +2133,9 @@ void FStaticLightingSystem::CalculateDirectSignedDistanceFieldLightingTextureMap
 				V0,
 				V1,
 				V2,
-				V0.Vertex.TextureCoordinates[TextureMapping->LightmapTextureCoordinateIndex] * FVector2D(HighResolutionSignalSizeX, HighResolutionSignalSizeY) + FVector2D(-0.5f,-0.5f),
-				V1.Vertex.TextureCoordinates[TextureMapping->LightmapTextureCoordinateIndex] * FVector2D(HighResolutionSignalSizeX, HighResolutionSignalSizeY) + FVector2D(-0.5f,-0.5f),
-				V2.Vertex.TextureCoordinates[TextureMapping->LightmapTextureCoordinateIndex] * FVector2D(HighResolutionSignalSizeX, HighResolutionSignalSizeY) + FVector2D(-0.5f,-0.5f),
+				V0.Vertex.TextureCoordinates[TextureMapping->LightmapTextureCoordinateIndex] * FVector2f(HighResolutionSignalSizeX, HighResolutionSignalSizeY) + FVector2f(-0.5f,-0.5f),
+				V1.Vertex.TextureCoordinates[TextureMapping->LightmapTextureCoordinateIndex] * FVector2f(HighResolutionSignalSizeX, HighResolutionSignalSizeY) + FVector2f(-0.5f,-0.5f),
+				V2.Vertex.TextureCoordinates[TextureMapping->LightmapTextureCoordinateIndex] * FVector2f(HighResolutionSignalSizeX, HighResolutionSignalSizeY) + FVector2f(-0.5f,-0.5f),
 				false
 				);
 		}
@@ -2219,12 +2219,12 @@ void FStaticLightingSystem::CalculateDirectSignedDistanceFieldLightingTextureMap
 							const bool bLightIsInFrontOfTriangle = !IsLightBehindSurface(HighResSample.GetPosition(),HighResSample.GetNormal(),Light);
 
 							if ((bLightIsInFrontOfTriangle || bIsTwoSided) 
-								&& Light->AffectsBounds(FBoxSphereBounds(HighResSample.GetPosition(), FVector4(0,0,0),0)))
+								&& Light->AffectsBounds(FBoxSphereBounds3f(HighResSample.GetPosition(), FVector4f(0,0,0),0)))
 							{
-								const FVector4 LightPosition = Light->LightCenterPosition(HighResSample.GetPosition(), HighResSample.GetNormal());
-								const FVector4 LightVector = (LightPosition - HighResSample.GetPosition()).GetSafeNormal();
+								const FVector4f LightPosition = Light->LightCenterPosition(HighResSample.GetPosition(), HighResSample.GetNormal());
+								const FVector4f LightVector = (LightPosition - HighResSample.GetPosition()).GetSafeNormal();
 
-								FVector4 NormalForOffset = HighResSample.GetNormal();
+								FVector4f NormalForOffset = HighResSample.GetNormal();
 								// Flip the normal used for offsetting the start of the ray for two sided materials if a flipped normal would be closer to the light.
 								// This prevents incorrect shadowing where using the frontface normal would cause the ray to start inside a nearby object.
 								if (bIsTwoSided && Dot3(-NormalForOffset, LightVector) > Dot3(NormalForOffset, LightVector))
@@ -2411,8 +2411,8 @@ void FStaticLightingSystem::CalculateDirectSignedDistanceFieldLightingTextureMap
 											if (LowResScatterSample.IsMapped())
 											{
 												bool CurrentRegion = false;
-												FVector4 ScatterPosition;
-												FVector4 ScatterNormal;
+												FVector4f ScatterPosition;
+												FVector4f ScatterNormal;
 												bool bFoundScatterPosition = false;
 												
 												if (LowResScatterSample.NeedsHighResSampling())
@@ -2511,11 +2511,11 @@ void FStaticLightingSystem::CalculateDirectSignedDistanceFieldLightingTextureMap
 	bool bDebugThisMapping,
 	const FLight* Light) const
 {
-	const FBoxSphereBounds MeshInfluenceBounds(TextureMapping->Mesh->BoundingBox.ExpandBy(ShadowSettings.MaxTransitionDistanceWorldSpace));
+	const FBoxSphereBounds3f MeshInfluenceBounds(TextureMapping->Mesh->BoundingBox.ExpandBy(ShadowSettings.MaxTransitionDistanceWorldSpace));
 
 	if (Light->AffectsBounds(MeshInfluenceBounds))
 	{
-		const FBoxSphereBounds SceneBounds = FBoxSphereBounds(AggregateMesh->GetBounds());
+		const FBoxSphereBounds3f SceneBounds = FBoxSphereBounds3f(AggregateMesh->GetBounds());
 		const FDirectionalLight* DirectionalLight = Light->GetDirectionalLight();
 		const FSpotLight* SpotLight = Light->GetSpotLight();
 		const FPointLight* PointLight = Light->GetPointLight();
@@ -2524,12 +2524,12 @@ void FStaticLightingSystem::CalculateDirectSignedDistanceFieldLightingTextureMap
 		if (DirectionalLight)
 		{
 			LIGHTINGSTAT(FManualRDTSCTimer FirstPassSourceTimer(MappingContext.Stats.SignedDistanceFieldSourceFirstPassThreadTime));
-			FVector4 XAxis, YAxis;
+			FVector4f XAxis, YAxis;
 			DirectionalLight->Direction.FindBestAxisVectors3(XAxis, YAxis);
 			// Create a coordinate system for the directional light, with the z axis corresponding to the light's direction
-			FMatrix WorldToLight = FBasisVectorMatrix(XAxis, YAxis, DirectionalLight->Direction, FVector4(0,0,0));
+			FMatrix44f WorldToLight = FBasisVectorMatrix44f(XAxis, YAxis, DirectionalLight->Direction, FVector4f(0,0,0));
 
-			const FBox LightSpaceImportanceBounds = MeshInfluenceBounds.GetBox().TransformBy(WorldToLight);
+			const FBox3f LightSpaceImportanceBounds = MeshInfluenceBounds.GetBox().TransformBy(WorldToLight);
 
 			FStaticShadowDepthMap ShadowDepthMap;
 
@@ -2554,7 +2554,7 @@ void FStaticLightingSystem::CalculateDirectSignedDistanceFieldLightingTextureMap
 			ShadowMap.Empty(ShadowMapSizeX * ShadowMapSizeY);
 			ShadowMap.AddZeroed(ShadowMapSizeX * ShadowMapSizeY);
 			const float ShadowMapStart = LightSpaceImportanceBounds.Max.Z - SceneBounds.SphereRadius * 2;
-			const FMatrix LightToWorld = WorldToLight.InverseFast();
+			const FMatrix44f LightToWorld = WorldToLight.InverseFast();
 
 			for (int32 Y = 0; Y < ShadowMapSizeY; Y++)
 			{
@@ -2564,14 +2564,14 @@ void FStaticLightingSystem::CalculateDirectSignedDistanceFieldLightingTextureMap
 				{
 					const float XFraction = (X + .5f) / (float)(ShadowMapSizeX - 1);
 
-					const FVector4 LightSpaceEndPosition(
+					const FVector4f LightSpaceEndPosition(
 						LightSpaceImportanceBounds.Min.X + XFraction * (LightSpaceImportanceBounds.Max.X - LightSpaceImportanceBounds.Min.X),
 						LightSpaceImportanceBounds.Min.Y + YFraction * (LightSpaceImportanceBounds.Max.Y - LightSpaceImportanceBounds.Min.Y),
 						LightSpaceImportanceBounds.Max.Z);
-					const FVector4 WorldSpaceEndPosition = LightToWorld.TransformPosition(LightSpaceEndPosition);
+					const FVector4f WorldSpaceEndPosition = LightToWorld.TransformPosition(LightSpaceEndPosition);
 
-					const FVector4 LightSpaceStartPosition(LightSpaceEndPosition.X, LightSpaceEndPosition.Y, ShadowMapStart);
-					const FVector4 WorldSpaceStartPosition = LightToWorld.TransformPosition(LightSpaceStartPosition);
+					const FVector4f LightSpaceStartPosition(LightSpaceEndPosition.X, LightSpaceEndPosition.Y, ShadowMapStart);
+					const FVector4f WorldSpaceStartPosition = LightToWorld.TransformPosition(LightSpaceStartPosition);
 
 					const FLightRay LightRay(
 						WorldSpaceStartPosition,
@@ -2624,8 +2624,8 @@ void FStaticLightingSystem::CalculateDirectSignedDistanceFieldLightingTextureMap
 
 					if (TexelToVertex.TotalSampleWeight > 0.0f)
 					{
-						const FVector4 LightSpacePosition = WorldToLight.TransformPosition(TexelToVertex.WorldPosition);
-						const FVector LightSpaceNormal = WorldToLight.TransformVector(TexelToVertex.WorldTangentZ);
+						const FVector4f LightSpacePosition = WorldToLight.TransformPosition(TexelToVertex.WorldPosition);
+						const FVector3f LightSpaceNormal = WorldToLight.TransformVector(TexelToVertex.WorldTangentZ);
 						const float SinThetaX = LightSpaceNormal.X;
 						const float TanThetaX = SinThetaX / FMath::Sqrt(1 - SinThetaX * SinThetaX);
 						const float SinThetaY = LightSpaceNormal.Y;
@@ -2648,7 +2648,7 @@ void FStaticLightingSystem::CalculateDirectSignedDistanceFieldLightingTextureMap
 						{
 							for (int32 SearchX = FMath::Max(ShadowMapX - TransitionSearchTexelRadiusX, 0); SearchX < FMath::Min(ShadowMapX + TransitionSearchTexelRadiusX, ShadowMapSizeX); SearchX++)
 							{
-								const FVector2D LightSpaceXYOffset((SearchX - ShadowMapX) * BoundsCellSizeX, (SearchY - ShadowMapY) * BoundsCellSizeY);
+								const FVector2f LightSpaceXYOffset((SearchX - ShadowMapX) * BoundsCellSizeX, (SearchY - ShadowMapY) * BoundsCellSizeY);
 								const float PlaneHeightOffsetX = LightSpaceXYOffset.X * TanThetaX;
 								const float PlaneHeightOffsetY = LightSpaceXYOffset.Y * TanThetaY;
 								const float PlaneHeightOffset = PlaneHeightOffsetX + PlaneHeightOffsetY;
@@ -2817,7 +2817,7 @@ void FStaticLightingSystem::CalculateDirectLightingTextureMappingPhotonMap(
 						else
 						{
 							// Fall back to using the UV's of this texel
-							CurrentVertex.TextureCoordinates[1] = FVector2D(X / (float)TextureMapping->CachedSizeX, Y / (float)TextureMapping->CachedSizeY);
+							CurrentVertex.TextureCoordinates[1] = FVector2f(X / (float)TextureMapping->CachedSizeX, Y / (float)TextureMapping->CachedSizeY);
 						}
 
 						// Find the nearest irradiance photon that was cached on this surface
@@ -2832,8 +2832,8 @@ void FStaticLightingSystem::CalculateDirectLightingTextureMappingPhotonMap(
 						FGatheredLightSample DirectLightingSample;
 						FGatheredLightSample Unused;
 						float Unused2;
-						TArray<FVector, TInlineAllocator<1>> VertexOffsets;
-						VertexOffsets.Add(FVector(0, 0, 0));
+						TArray<FVector3f, TInlineAllocator<1>> VertexOffsets;
+						VertexOffsets.Add(FVector3f(0, 0, 0));
 
 						CalculateApproximateDirectLighting(CurrentVertex, TexelToVertex.TexelRadius, VertexOffsets, .1f, true, true, bDebugThisTexel, MappingContext, DirectLightingSample, Unused, Unused2);
 
@@ -2851,7 +2851,7 @@ void FStaticLightingSystem::CalculateDirectLightingTextureMappingPhotonMap(
 
 						//@todo - can't visualize accurately using AmbientLight with directional lightmaps
 						//CurrentLightSample.AddWeighted(FGatheredLightSampleUtil::AmbientLight<2>(FinalLighting), 1.0f);
-						CurrentLightSample.AddWeighted(FGatheredLightSampleUtil::PointLightWorldSpace<2>(FinalLighting, FVector4(0, 0, 1), CurrentVertex.WorldTangentZ), 1.0f);
+						CurrentLightSample.AddWeighted(FGatheredLightSampleUtil::PointLightWorldSpace<2>(FinalLighting, FVector4f(0, 0, 1), CurrentVertex.WorldTangentZ), 1.0f);
 					}
 				}
 				else
@@ -2900,7 +2900,7 @@ void FStaticLightingSystem::ProcessCacheIndirectLightingTask(FCacheIndirectTaskD
 				const FTexelToVertex& TexelToVertex = (*Task->TexelToVertexMap)(X,Y);
 				checkSlow(TexelToVertex.TotalSampleWeight > 0.0f);
 				FFullStaticLightingVertex TexelVertex = TexelToVertex.GetFullVertex();
-				TexelVertex.TextureCoordinates[1] = FVector2D(X / (float)Task->TextureMapping->CachedSizeX, Y / (float)Task->TextureMapping->CachedSizeY);
+				TexelVertex.TextureCoordinates[1] = FVector2f(X / (float)Task->TextureMapping->CachedSizeX, Y / (float)Task->TextureMapping->CachedSizeY);
 
 				// Calculate incoming radiance for the frontface
 				FGatheredLightSample IndirectLightingSample = CachePointIncomingRadiance(
@@ -3427,9 +3427,9 @@ void FStaticLightingSystem::ColorInvalidLightmapUVs(
 		int32 DummyElement;
 		TextureMapping->Mesh->GetTriangle(TriangleIndex,V0,V1,V2,DummyElement);
 
-		const FVector2D UV0 = V0.TextureCoordinates[TextureMapping->LightmapTextureCoordinateIndex];
-		const FVector2D UV1 = V1.TextureCoordinates[TextureMapping->LightmapTextureCoordinateIndex];
-		const FVector2D UV2 = V2.TextureCoordinates[TextureMapping->LightmapTextureCoordinateIndex];
+		const FVector2f UV0 = V0.TextureCoordinates[TextureMapping->LightmapTextureCoordinateIndex];
+		const FVector2f UV1 = V1.TextureCoordinates[TextureMapping->LightmapTextureCoordinateIndex];
+		const FVector2f UV2 = V2.TextureCoordinates[TextureMapping->LightmapTextureCoordinateIndex];
 
 		bool bHasWrappingLightmapUVs = false;
 		//@todo - remove the thresholds and fixup existing content
@@ -3448,9 +3448,9 @@ void FStaticLightingSystem::ColorInvalidLightmapUVs(
 			bHasWrappingLightmapUVs,
 			bHasWrappingLightmapUVs,
 			bHasWrappingLightmapUVs,
-			UV0 * FVector2D(TextureMapping->CachedSizeX,TextureMapping->CachedSizeY) + FVector2D(-0.5f,-0.5f),
-			UV1 * FVector2D(TextureMapping->CachedSizeX,TextureMapping->CachedSizeY) + FVector2D(-0.5f,-0.5f),
-			UV2 * FVector2D(TextureMapping->CachedSizeX,TextureMapping->CachedSizeY) + FVector2D(-0.5f,-0.5f),
+			UV0 * FVector2f(TextureMapping->CachedSizeX,TextureMapping->CachedSizeY) + FVector2f(-0.5f,-0.5f),
+			UV1 * FVector2f(TextureMapping->CachedSizeX,TextureMapping->CachedSizeY) + FVector2f(-0.5f,-0.5f),
+			UV2 * FVector2f(TextureMapping->CachedSizeX,TextureMapping->CachedSizeY) + FVector2f(-0.5f,-0.5f),
 			false 
 			);
 	}
@@ -3704,12 +3704,12 @@ void FStaticLightingSystem::PadTextureMapping(
 /** Rasterizes Mesh into TexelToCornersMap */
 void FStaticLightingSystem::CalculateTexelCorners(const FStaticLightingMesh* Mesh, FTexelToCornersMap& TexelToCornersMap, int32 UVIndex, bool bDebugThisMapping) const
 {
-	static const FVector2D CornerOffsets[NumTexelCorners] = 
+	static const FVector2f CornerOffsets[NumTexelCorners] = 
 	{
-		FVector2D(0, 0),
-		FVector2D(-1, 0),
-		FVector2D(0, -1),
-		FVector2D(-1, -1)
+		FVector2f(0, 0),
+		FVector2f(-1, 0),
+		FVector2f(0, -1),
+		FVector2f(-1, -1)
 	};
 
 	// Rasterize each triangle of the mesh
@@ -3736,9 +3736,9 @@ void FStaticLightingSystem::CalculateTexelCorners(const FStaticLightingMesh* Mes
 				V0,
 				V1,
 				V2,
-				V0.TextureCoordinates[UVIndex] * FVector2D(TexelToCornersMap.GetSizeX(), TexelToCornersMap.GetSizeY()) + CornerOffsets[CornerIndex],
-				V1.TextureCoordinates[UVIndex] * FVector2D(TexelToCornersMap.GetSizeX(), TexelToCornersMap.GetSizeY()) + CornerOffsets[CornerIndex],
-				V2.TextureCoordinates[UVIndex] * FVector2D(TexelToCornersMap.GetSizeX(), TexelToCornersMap.GetSizeY()) + CornerOffsets[CornerIndex],
+				V0.TextureCoordinates[UVIndex] * FVector2f(TexelToCornersMap.GetSizeX(), TexelToCornersMap.GetSizeY()) + CornerOffsets[CornerIndex],
+				V1.TextureCoordinates[UVIndex] * FVector2f(TexelToCornersMap.GetSizeX(), TexelToCornersMap.GetSizeY()) + CornerOffsets[CornerIndex],
+				V2.TextureCoordinates[UVIndex] * FVector2f(TexelToCornersMap.GetSizeX(), TexelToCornersMap.GetSizeY()) + CornerOffsets[CornerIndex],
 				false
 				);
 		}
@@ -3754,15 +3754,15 @@ void FStaticLightingSystem::CalculateTexelCorners(
 	int32 MaterialIndex,
 	int32 UVIndex, 
 	bool bDebugThisMapping, 
-	FVector2D UVBias, 
-	FVector2D UVScale) const
+	FVector2f UVBias, 
+	FVector2f UVScale) const
 {
-	static const FVector2D CornerOffsets[NumTexelCorners] = 
+	static const FVector2f CornerOffsets[NumTexelCorners] = 
 	{
-		FVector2D(0, 0),
-		FVector2D(-1, 0),
-		FVector2D(0, -1),
-		FVector2D(-1, -1)
+		FVector2f(0, 0),
+		FVector2f(-1, 0),
+		FVector2f(0, -1),
+		FVector2f(-1, -1)
 	};
 
 	// Rasterize each triangle of the mesh
@@ -3788,9 +3788,9 @@ void FStaticLightingSystem::CalculateTexelCorners(
 					V0,
 					V1,
 					V2,
-					UVScale * (UVBias + V0.TextureCoordinates[UVIndex]) * FVector2D(TexelToCornersMap.GetSizeX(), TexelToCornersMap.GetSizeY()) + CornerOffsets[CornerIndex],
-					UVScale * (UVBias + V1.TextureCoordinates[UVIndex]) * FVector2D(TexelToCornersMap.GetSizeX(), TexelToCornersMap.GetSizeY()) + CornerOffsets[CornerIndex],
-					UVScale * (UVBias + V2.TextureCoordinates[UVIndex]) * FVector2D(TexelToCornersMap.GetSizeX(), TexelToCornersMap.GetSizeY()) + CornerOffsets[CornerIndex],
+					UVScale * (UVBias + V0.TextureCoordinates[UVIndex]) * FVector2f(TexelToCornersMap.GetSizeX(), TexelToCornersMap.GetSizeY()) + CornerOffsets[CornerIndex],
+					UVScale * (UVBias + V1.TextureCoordinates[UVIndex]) * FVector2f(TexelToCornersMap.GetSizeX(), TexelToCornersMap.GetSizeY()) + CornerOffsets[CornerIndex],
+					UVScale * (UVBias + V2.TextureCoordinates[UVIndex]) * FVector2f(TexelToCornersMap.GetSizeX(), TexelToCornersMap.GetSizeY()) + CornerOffsets[CornerIndex],
 					false
 					);
 			}
diff --git a/Engine/Source/Programs/UnrealLightmass/Private/Lighting/TextureMappingSetup.h b/Engine/Source/Programs/UnrealLightmass/Private/Lighting/TextureMappingSetup.h
--- a/Engine/Source/Programs/UnrealLightmass/Private/Lighting/TextureMappingSetup.h
+++ b/Engine/Source/Programs/UnrealLightmass/Private/Lighting/TextureMappingSetup.h
@@ -11,11 +11,11 @@ namespace Lightmass
 /** A map from a texel to the world-space surface point which maps the texel. */
 struct FTexelToVertex
 {
-	FVector4 WorldPosition;
-	FVector4 WorldTangentX;
-	FVector4 WorldTangentY;
-	FVector4 WorldTangentZ;
-	FVector4 TriangleNormal;
+	FVector4f WorldPosition;
+	FVector4f WorldTangentX;
+	FVector4f WorldTangentY;
+	FVector4f WorldTangentZ;
+	FVector4f TriangleNormal;
 
 	/** Weight used when combining super sampled attributes and determining if the texel has been mapped. */
 	float TotalSampleWeight;
@@ -35,7 +35,7 @@ struct FTexelToVertex
 	uint16 ElementIndex;
 
 	/** Texture coordinates */
-	FVector2D TextureCoordinates[MAX_TEXCOORDS];
+	FVector2f TextureCoordinates[MAX_TEXCOORDS];
 
 	/** Create a static lighting vertex to represent the texel. */
 	inline FStaticLightingVertex GetVertex() const
@@ -162,7 +162,7 @@ public:
 		const FScene& InScene,
 		FTexelToVertexMap& InTexelToVertexMap,
 		float InSampleWeight,
-		const FVector4& InTriangleNormal,
+		const FVector4f& InTriangleNormal,
 		bool bInDebugThisMapping,
 		bool bInUseMaxWeight
 		) :
@@ -194,7 +194,7 @@ private:
 
 	/** The weight of the current sample. */
 	const float SampleWeight;
-	const FVector4 TriangleNormal;
+	const FVector4f TriangleNormal;
 	const bool bDebugThisMapping;
 	const bool bUseMaxWeight;
 };
diff --git a/Engine/Source/Programs/UnrealLightmass/Private/Lighting/VolumeDistanceField.cpp b/Engine/Source/Programs/UnrealLightmass/Private/Lighting/VolumeDistanceField.cpp
--- a/Engine/Source/Programs/UnrealLightmass/Private/Lighting/VolumeDistanceField.cpp
+++ b/Engine/Source/Programs/UnrealLightmass/Private/Lighting/VolumeDistanceField.cpp
@@ -17,15 +17,15 @@ void FStaticLightingSystem::BeginCalculateVolumeDistanceField()
 		DistanceFieldVolumeBounds = AggregateMesh->GetBounds();
 	}
 
-	FBox UnclampedDistanceFieldVolumeBounds = DistanceFieldVolumeBounds;
-	FVector4 DoubleExtent = UnclampedDistanceFieldVolumeBounds.GetExtent() * 2;
-	DoubleExtent.X = DoubleExtent.X - FMath::Fmod(DoubleExtent.X, (FVector4::FReal)VolumeDistanceFieldSettings.VoxelSize) + VolumeDistanceFieldSettings.VoxelSize;
-	DoubleExtent.Y = DoubleExtent.Y - FMath::Fmod(DoubleExtent.Y, (FVector4::FReal)VolumeDistanceFieldSettings.VoxelSize) + VolumeDistanceFieldSettings.VoxelSize;
-	DoubleExtent.Z = DoubleExtent.Z - FMath::Fmod(DoubleExtent.Z, (FVector4::FReal)VolumeDistanceFieldSettings.VoxelSize) + VolumeDistanceFieldSettings.VoxelSize;
+	FBox3f UnclampedDistanceFieldVolumeBounds = DistanceFieldVolumeBounds;
+	FVector4f DoubleExtent = UnclampedDistanceFieldVolumeBounds.GetExtent() * 2;
+	DoubleExtent.X = DoubleExtent.X - FMath::Fmod(DoubleExtent.X, (FVector4f::FReal)VolumeDistanceFieldSettings.VoxelSize) + VolumeDistanceFieldSettings.VoxelSize;
+	DoubleExtent.Y = DoubleExtent.Y - FMath::Fmod(DoubleExtent.Y, (FVector4f::FReal)VolumeDistanceFieldSettings.VoxelSize) + VolumeDistanceFieldSettings.VoxelSize;
+	DoubleExtent.Z = DoubleExtent.Z - FMath::Fmod(DoubleExtent.Z, (FVector4f::FReal)VolumeDistanceFieldSettings.VoxelSize) + VolumeDistanceFieldSettings.VoxelSize;
 	// Round the max up to the next step boundary
 	UnclampedDistanceFieldVolumeBounds.Max = UnclampedDistanceFieldVolumeBounds.Min + DoubleExtent;
 
-	const FVector4 VolumeSizes = UnclampedDistanceFieldVolumeBounds.GetExtent() * 2.0f / VolumeDistanceFieldSettings.VoxelSize;
+	const FVector4f VolumeSizes = UnclampedDistanceFieldVolumeBounds.GetExtent() * 2.0f / VolumeDistanceFieldSettings.VoxelSize;
 	VolumeSizeX = FMath::TruncToFloat(VolumeSizes.X + DELTA);
 	VolumeSizeY = FMath::TruncToFloat(VolumeSizes.Y + DELTA);
 	VolumeSizeZ = FMath::TruncToFloat(VolumeSizes.Z + DELTA);
@@ -42,13 +42,13 @@ void FStaticLightingSystem::BeginCalculateVolumeDistanceField()
 		DistanceFieldVoxelSize = VolumeDistanceFieldSettings.VoxelSize * SingleDimensionScale;
 
 		DoubleExtent = DistanceFieldVolumeBounds.GetExtent() * 2;
-		DoubleExtent.X = DoubleExtent.X - FMath::Fmod(DoubleExtent.X, (FVector4::FReal)DistanceFieldVoxelSize) + DistanceFieldVoxelSize;
-		DoubleExtent.Y = DoubleExtent.Y - FMath::Fmod(DoubleExtent.Y, (FVector4::FReal)DistanceFieldVoxelSize) + DistanceFieldVoxelSize;
-		DoubleExtent.Z = DoubleExtent.Z - FMath::Fmod(DoubleExtent.Z, (FVector4::FReal)DistanceFieldVoxelSize) + DistanceFieldVoxelSize;
+		DoubleExtent.X = DoubleExtent.X - FMath::Fmod(DoubleExtent.X, (FVector4f::FReal)DistanceFieldVoxelSize) + DistanceFieldVoxelSize;
+		DoubleExtent.Y = DoubleExtent.Y - FMath::Fmod(DoubleExtent.Y, (FVector4f::FReal)DistanceFieldVoxelSize) + DistanceFieldVoxelSize;
+		DoubleExtent.Z = DoubleExtent.Z - FMath::Fmod(DoubleExtent.Z, (FVector4f::FReal)DistanceFieldVoxelSize) + DistanceFieldVoxelSize;
 		// Round the max up to the next step boundary with the clamped voxel size
 		DistanceFieldVolumeBounds.Max = DistanceFieldVolumeBounds.Min + DoubleExtent;
 
-		const FVector4 ClampedVolumeSizes = DistanceFieldVolumeBounds.GetExtent() * 2.0f / DistanceFieldVoxelSize;
+		const FVector4f ClampedVolumeSizes = DistanceFieldVolumeBounds.GetExtent() * 2.0f / DistanceFieldVoxelSize;
 		VolumeSizeX = FMath::TruncToFloat(ClampedVolumeSizes.X + DELTA);
 		VolumeSizeY = FMath::TruncToFloat(ClampedVolumeSizes.Y + DELTA);
 		VolumeSizeZ = FMath::TruncToFloat(ClampedVolumeSizes.Z + DELTA);
@@ -74,29 +74,29 @@ void FStaticLightingSystem::CalculateVolumeDistanceFieldWorkRange(int32 ZIndex)
 	const double StartTime = FPlatformTime::Seconds();
 	FStaticLightingMappingContext MappingContext(NULL, *this);
 
-	TArray<FVector4> SampleDirections;
-	TArray<FVector2D> SampleDirectionUniforms;
+	TArray<FVector4f> SampleDirections;
+	TArray<FVector2f> SampleDirectionUniforms;
 	const int32 NumThetaSteps = FMath::TruncToInt(FMath::Sqrt(VolumeDistanceFieldSettings.NumVoxelDistanceSamples / (2.0f * (float)PI)));
 	const int32 NumPhiSteps = FMath::TruncToInt(NumThetaSteps * (float)PI);
 	FLMRandomStream RandomStream(0);
 	GenerateStratifiedUniformHemisphereSamples(NumThetaSteps, NumPhiSteps, RandomStream, SampleDirections, SampleDirectionUniforms);
-	TArray<FVector4> OtherHemisphereSamples;
-	TArray<FVector2D> OtherSampleDirectionUniforms;
+	TArray<FVector4f> OtherHemisphereSamples;
+	TArray<FVector2f> OtherSampleDirectionUniforms;
 	GenerateStratifiedUniformHemisphereSamples(NumThetaSteps, NumPhiSteps, RandomStream, OtherHemisphereSamples, OtherSampleDirectionUniforms);
 
 	for (int32 i = 0; i < OtherHemisphereSamples.Num(); i++)
 	{
-		FVector4 Sample = OtherHemisphereSamples[i];
+		FVector4f Sample = OtherHemisphereSamples[i];
 		Sample.Z *= -1;
 		SampleDirections.Add(Sample);
 	}
 
-	const FVector4 CellExtents = FVector4(DistanceFieldVoxelSize / 2, DistanceFieldVoxelSize / 2, DistanceFieldVoxelSize / 2);
+	const FVector4f CellExtents = FVector4f(DistanceFieldVoxelSize / 2, DistanceFieldVoxelSize / 2, DistanceFieldVoxelSize / 2);
 	for (int32 YIndex = 0; YIndex < VolumeSizeY; YIndex++)
 	{
 		for (int32 XIndex = 0; XIndex < VolumeSizeX; XIndex++)
 		{
-			const FVector4 VoxelPosition = FVector4(XIndex, YIndex, ZIndex) * DistanceFieldVoxelSize + DistanceFieldVolumeBounds.Min + CellExtents;
+			const FVector4f VoxelPosition = FVector4f(XIndex, YIndex, ZIndex) * DistanceFieldVoxelSize + DistanceFieldVolumeBounds.Min + CellExtents;
 			const int32 Index = ZIndex * VolumeSizeY * VolumeSizeX + YIndex * VolumeSizeX + XIndex;
 	
 			float MinDistance[2];
diff --git a/Engine/Source/Programs/UnrealLightmass/Private/LightmassCore/Math/LMCollision.h b/Engine/Source/Programs/UnrealLightmass/Private/LightmassCore/Math/LMCollision.h
--- a/Engine/Source/Programs/UnrealLightmass/Private/LightmassCore/Math/LMCollision.h
+++ b/Engine/Source/Programs/UnrealLightmass/Private/LightmassCore/Math/LMCollision.h
@@ -7,12 +7,12 @@
 namespace Lightmass
 {
 
-static bool ClipLineWithBox(const FBox& Box, const FVector4& Start, const FVector4& End, FVector4& IntersectedStart, FVector4& IntersectedEnd)
+static bool ClipLineWithBox(const FBox3f& Box, const FVector4f& Start, const FVector4f& End, FVector4f& IntersectedStart, FVector4f& IntersectedEnd)
 {
     IntersectedStart = Start;
     IntersectedEnd = End;
 
-    FVector4 Dir;
+    FVector4f Dir;
     float TEdgeOfBox,TLineLength;
     bool StartCulled,EndCulled;
    
diff --git a/Engine/Source/Programs/UnrealLightmass/Private/LightmassCore/Math/LMMath.cpp b/Engine/Source/Programs/UnrealLightmass/Private/LightmassCore/Math/LMMath.cpp
--- a/Engine/Source/Programs/UnrealLightmass/Private/LightmassCore/Math/LMMath.cpp
+++ b/Engine/Source/Programs/UnrealLightmass/Private/LightmassCore/Math/LMMath.cpp
@@ -12,17 +12,17 @@ FLinearColor FLinearColorUtils::XYZToLinearRGB(const FLinearColor& InColor)
 {
 	FLinearColor SourceXYZ(InColor);
 	// Inverse of the transform in FLinearColor::LinearRGBToXYZ()
-	const FMatrix XYZToRGB(
-		FVector(3.2404548f, -0.9692664f, 0.0556434f),
-		FVector(-1.5371389f, 1.8760109f, -0.2040259f),
-		FVector(-0.4985315f, 0.0415561f, 1.0572252f),
-		FVector(0,			  0,		  0)); 
+	const FMatrix44f XYZToRGB(
+		FVector3f(3.2404548f, -0.9692664f, 0.0556434f),
+		FVector3f(-1.5371389f, 1.8760109f, -0.2040259f),
+		FVector3f(-0.4985315f, 0.0415561f, 1.0572252f),
+		FVector3f(0,			  0,		  0)); 
 
 	if (FMath::IsNearlyEqual(InColor.R, 0.0f, (float)SMALL_NUMBER) && FMath::IsNearlyEqual(InColor.B, 0.0f, (float)SMALL_NUMBER))
 	{
 		SourceXYZ.G = 0.0f;
 	}
-	const FVector4 LinearRGB = XYZToRGB.TransformVector(FVector(SourceXYZ.R, SourceXYZ.G, SourceXYZ.B));
+	const FVector4f LinearRGB = XYZToRGB.TransformVector(FVector3f(SourceXYZ.R, SourceXYZ.G, SourceXYZ.B));
 	return FLinearColor(FMath::Max(LinearRGB.X, 0.0f), FMath::Max(LinearRGB.Y, 0.0f), FMath::Max(LinearRGB.Z, 0.0f));
 }
 
@@ -103,13 +103,13 @@ FLinearColor FLinearColorUtils::LinearRGBToXYZ(const FLinearColor& InColor)
 {
 	// RGB to XYZ linear transformation used by sRGB
 	//http://www.w3.org/Graphics/Color/sRGB
-	const FMatrix RGBToXYZ(
-		FVector(0.4124564f, 0.2126729f, 0.0193339f),
-		FVector(0.3575761f, 0.7151522f, 0.1191920f),
-		FVector(0.1804375f, 0.0721750f, 0.9503041f),
-		FVector(0,			 0,			 0)); 
+	const FMatrix44f RGBToXYZ(
+		FVector3f(0.4124564f, 0.2126729f, 0.0193339f),
+		FVector3f(0.3575761f, 0.7151522f, 0.1191920f),
+		FVector3f(0.1804375f, 0.0721750f, 0.9503041f),
+		FVector3f(0,			 0,			 0)); 
 
-	const FVector4 ResultVector = RGBToXYZ.TransformVector(FVector(InColor.R, InColor.G, InColor.B));
+	const FVector4f ResultVector = RGBToXYZ.TransformVector(FVector3f(InColor.R, InColor.G, InColor.B));
 	return FLinearColor(ResultVector.X, ResultVector.Y, ResultVector.Z);
 }
 
@@ -135,18 +135,18 @@ FLinearColor FLinearColorUtils::AdjustSaturation(const FLinearColor& InColor, fl
 }
 
 bool GetBarycentricWeights(
-	const FVector4& Position0,
-	const FVector4& Position1,
-	const FVector4& Position2,
-	const FVector4& InterpolatePosition,
+	const FVector4f& Position0,
+	const FVector4f& Position1,
+	const FVector4f& Position2,
+	const FVector4f& InterpolatePosition,
 	float Tolerance,
-	FVector4& BarycentricWeights
+	FVector4f& BarycentricWeights
 	)
 {
-	BarycentricWeights = FVector4(0,0,0);
-	FVector4 TriangleNormal = (Position0 - Position1) ^ (Position2 - Position0);
+	BarycentricWeights = FVector4f(0,0,0);
+	FVector4f TriangleNormal = (Position0 - Position1) ^ (Position2 - Position0);
 	float ParallelogramArea = TriangleNormal.Size3();
-	FVector4 UnitTriangleNormal = TriangleNormal / ParallelogramArea;
+	FVector4f UnitTriangleNormal = TriangleNormal / ParallelogramArea;
 	float PlaneDistance = Dot3(UnitTriangleNormal, (InterpolatePosition - Position0));
 
 	// Only continue if the position to interpolate to is in the plane of the triangle (within some error)
@@ -154,21 +154,21 @@ bool GetBarycentricWeights(
 	{
 		// Move the position to interpolate to into the plane of the triangle along the normal, 
 		// Otherwise there will be error in our barycentric coordinates
-		FVector4 AdjustedInterpolatePosition = InterpolatePosition - UnitTriangleNormal * PlaneDistance;
+		FVector4f AdjustedInterpolatePosition = InterpolatePosition - UnitTriangleNormal * PlaneDistance;
 
-		FVector4 NormalU = (AdjustedInterpolatePosition - Position1) ^ (Position2 - AdjustedInterpolatePosition);
+		FVector4f NormalU = (AdjustedInterpolatePosition - Position1) ^ (Position2 - AdjustedInterpolatePosition);
 		// Signed area, if negative then InterpolatePosition is not in the triangle
 		float ParallelogramAreaU = NormalU.Size3() * (Dot3(NormalU, TriangleNormal) > 0.0f ? 1.0f : -1.0f);
 		float BaryCentricU = ParallelogramAreaU / ParallelogramArea;
 
-		FVector4 NormalV = (AdjustedInterpolatePosition - Position2) ^ (Position0 - AdjustedInterpolatePosition);
+		FVector4f NormalV = (AdjustedInterpolatePosition - Position2) ^ (Position0 - AdjustedInterpolatePosition);
 		float ParallelogramAreaV = NormalV.Size3() * (Dot3(NormalV, TriangleNormal) > 0.0f ? 1.0f : -1.0f);
 		float BaryCentricV = ParallelogramAreaV / ParallelogramArea;
 
 		float BaryCentricW = 1.0f - BaryCentricU - BaryCentricV;
 		if (BaryCentricU > -Tolerance && BaryCentricV > -Tolerance && BaryCentricW > -Tolerance)
 		{
-			BarycentricWeights = FVector4(BaryCentricU, BaryCentricV, BaryCentricW);
+			BarycentricWeights = FVector4f(BaryCentricU, BaryCentricV, BaryCentricW);
 			return true;
 		}
 	}
diff --git a/Engine/Source/Programs/UnrealLightmass/Private/LightmassCore/Math/LMMath.h b/Engine/Source/Programs/UnrealLightmass/Private/LightmassCore/Math/LMMath.h
--- a/Engine/Source/Programs/UnrealLightmass/Private/LightmassCore/Math/LMMath.h
+++ b/Engine/Source/Programs/UnrealLightmass/Private/LightmassCore/Math/LMMath.h
@@ -16,12 +16,12 @@ namespace Lightmass
 #endif
 
 bool GetBarycentricWeights(
-	const FVector4& Position0,
-	const FVector4& Position1,
-	const FVector4& Position2,
-	const FVector4& InterpolatePosition,
+	const FVector4f& Position0,
+	const FVector4f& Position1,
+	const FVector4f& Position2,
+	const FVector4f& InterpolatePosition,
 	float Tolerance,
-	FVector4& BarycentricWeights
+	FVector4f& BarycentricWeights
 	);
 
 
@@ -442,7 +442,7 @@ struct FLinearColorUtils
 struct FVectorUtils
 {
 	/** @return >0: point is in front of the plane, <0: behind, =0: on the plane */
-	FORCEINLINE static float PlaneDot( const FVector4 &V, const FVector4 &P )
+	FORCEINLINE static float PlaneDot( const FVector4f &V, const FVector4f &P )
 	{
 		return V.X*P.X + V.Y*P.Y + V.Z*P.Z - V.W;
 	}
@@ -480,10 +480,10 @@ FORCEINLINE uint32 appCountTrailingZeros(uint32 Value)
 #endif // PLATFORM_WINDOWS
 
 /** Converts spherical coordinates on the unit sphere into a cartesian unit length vector. */
-FORCEINLINE FVector4 SphericalToUnitCartesian(const FVector2D& InHemispherical)
+FORCEINLINE FVector4f SphericalToUnitCartesian(const FVector2f& InHemispherical)
 {
 	const float SinTheta = FMath::Sin(InHemispherical.X);
-	return FVector4(FMath::Cos(InHemispherical.Y) * SinTheta, FMath::Sin(InHemispherical.Y) * SinTheta, FMath::Cos(InHemispherical.X));
+	return FVector4f(FMath::Cos(InHemispherical.Y) * SinTheta, FMath::Sin(InHemispherical.Y) * SinTheta, FMath::Cos(InHemispherical.X));
 }
 
 }
diff --git a/Engine/Source/Programs/UnrealLightmass/Private/LightmassCore/Math/LMMathSSE.h b/Engine/Source/Programs/UnrealLightmass/Private/LightmassCore/Math/LMMathSSE.h
--- a/Engine/Source/Programs/UnrealLightmass/Private/LightmassCore/Math/LMMathSSE.h
+++ b/Engine/Source/Programs/UnrealLightmass/Private/LightmassCore/Math/LMMathSSE.h
@@ -421,10 +421,10 @@ FORCEINLINE void LmVectorMatrixMultiply( void *Result, const void* Matrix1, cons
 }
 
 /**
- * Calculate the inverse of an FMatrix.
+ * Calculate the inverse of an FMatrix44f.
  *
- * @param DstMatrix		FMatrix pointer to where the result should be stored
- * @param SrcMatrix		FMatrix pointer to the Matrix to be inversed
+ * @param DstMatrix		FMatrix44f pointer to where the result should be stored
+ * @param SrcMatrix		FMatrix44f pointer to the Matrix to be inversed
  */
 #define LmVectorMatrixInverse( DstMatrix, SrcMatrix ) D3DXMatrixInverse( (D3DXMATRIX*) DstMatrix, NULL, (const D3DXMATRIX*)SrcMatrix )
 
diff --git a/Engine/Source/Programs/UnrealLightmass/Private/LightmassCore/Math/LMOctree.h b/Engine/Source/Programs/UnrealLightmass/Private/LightmassCore/Math/LMOctree.h
--- a/Engine/Source/Programs/UnrealLightmass/Private/LightmassCore/Math/LMOctree.h
+++ b/Engine/Source/Programs/UnrealLightmass/Private/LightmassCore/Math/LMOctree.h
@@ -19,41 +19,41 @@ namespace Lightmass
 class FBoxCenterAndExtent
 {
 public:
-	FVector4 Center;
-	FVector4 Extent;
+	FVector4f Center;
+	FVector4f Extent;
 
 	/** Default constructor. */
 	FBoxCenterAndExtent() {}
 
 	/** Initialization constructor. */
-	FBoxCenterAndExtent(const FVector4& InCenter,const FVector4& InExtent)
+	FBoxCenterAndExtent(const FVector4f& InCenter,const FVector4f& InExtent)
 	:	Center(InCenter)
 	,	Extent(InExtent)
 	{}
 
-	/** FBox conversion constructor. */
-	FBoxCenterAndExtent(const FBox& Box)
+	/** FBox3f conversion constructor. */
+	FBoxCenterAndExtent(const FBox3f& Box)
 	{
-		FVector BoxCenter;
-		FVector BoxExtent;
+		FVector3f BoxCenter;
+		FVector3f BoxExtent;
 		Box.GetCenterAndExtents(BoxCenter, BoxExtent);
 		Center = BoxCenter;
 		Extent = BoxExtent;
 		Center.W = Extent.W = 0;
 	}
 
-	/** FBoxSphereBounds conversion constructor. */
-	FBoxCenterAndExtent(const FBoxSphereBounds& BoxSphere)
+	/** FBoxSphereBounds3f conversion constructor. */
+	FBoxCenterAndExtent(const FBoxSphereBounds3f& BoxSphere)
 	{
 		Center = BoxSphere.Origin;
 		Extent = BoxSphere.BoxExtent;
 		Center.W = Extent.W = 0;
 	}
 
-	/** Converts to a FBox. */
-	FORCEINLINE FBox GetBox() const
+	/** Converts to a FBox3f. */
+	FORCEINLINE FBox3f GetBox() const
 	{
-		return FBox(Center - Extent,Center + Extent);
+		return FBox3f(Center - Extent,Center + Extent);
 	}
 
 	/**
@@ -63,12 +63,12 @@ public:
 	friend bool Intersect(const FBoxCenterAndExtent& A,const FBoxCenterAndExtent& B)
 	{
 		// Note: The W component of the bounds center positions is unreliable here, and must not be used in the intersection
-		FVector4 CenterDifference = A.Center - B.Center;
+		FVector4f CenterDifference = A.Center - B.Center;
 		CenterDifference.X = FMath::Abs(CenterDifference.X);
 		CenterDifference.Y = FMath::Abs(CenterDifference.Y);
 		CenterDifference.Z = FMath::Abs(CenterDifference.Z);
 
-		const FVector4 CompositeExtent = A.Extent + B.Extent;
+		const FVector4f CompositeExtent = A.Extent + B.Extent;
 
 		return CenterDifference.X < CompositeExtent.X && CenterDifference.Y < CompositeExtent.Y && CenterDifference.Z < CompositeExtent.Z;
 	}
@@ -215,13 +215,13 @@ public:
 	FORCEINLINE FOctreeNodeContext GetChildContext(FOctreeChildNodeRef ChildRef) const
 	{
 		return FOctreeNodeContext(FBoxCenterAndExtent(
-				FVector4(
+				FVector4f(
 					Bounds.Center.X + ChildCenterOffset * (-1.0f + 2 * ChildRef.X),
 					Bounds.Center.Y + ChildCenterOffset * (-1.0f + 2 * ChildRef.Y),
 					Bounds.Center.Z + ChildCenterOffset * (-1.0f + 2 * ChildRef.Z),
 					0
 					),
-				FVector4(
+				FVector4f(
 					ChildExtent,
 					ChildExtent,
 					ChildExtent,
@@ -607,7 +607,7 @@ public:
 	void GetMemoryUsage(size_t& OutSizeBytes) const;
 
 	/** Initialization constructor. */
-	TOctree(const FVector4& InOrigin, float InExtent);
+	TOctree(const FVector4f& InOrigin, float InExtent);
 
 	void Destroy()
 	{
@@ -887,9 +887,9 @@ void TOctree<ElementType,OctreeSemantics>::GetMemoryUsage(size_t& OutSizeBytes)
 }
 
 template<typename ElementType,typename OctreeSemantics>
-TOctree<ElementType,OctreeSemantics>::TOctree(const FVector4& InOrigin,float InExtent)
+TOctree<ElementType,OctreeSemantics>::TOctree(const FVector4f& InOrigin,float InExtent)
 :	RootNode(NULL)
-,	RootNodeContext(FBoxCenterAndExtent(InOrigin,FVector4(InExtent,InExtent,InExtent,0)), OctreeSemantics::LoosenessDenominator)
+,	RootNodeContext(FBoxCenterAndExtent(InOrigin,FVector4f(InExtent,InExtent,InExtent,0)), OctreeSemantics::LoosenessDenominator)
 ,	MinLeafExtent(InExtent * FMath::Pow((1.0f + 1.0f / (float)OctreeSemantics::LoosenessDenominator) / 2.0f,OctreeSemantics::MaxNodeDepth))
 {
 }
diff --git a/Engine/Source/Programs/UnrealLightmass/Private/LightmassCore/Math/LMkDOP.h b/Engine/Source/Programs/UnrealLightmass/Private/LightmassCore/Math/LMkDOP.h
--- a/Engine/Source/Programs/UnrealLightmass/Private/LightmassCore/Math/LMkDOP.h
+++ b/Engine/Source/Programs/UnrealLightmass/Private/LightmassCore/Math/LMkDOP.h
@@ -25,7 +25,7 @@ namespace Lightmass
 struct FHitResult
 {
 	/** Normal vector in coordinate system of the returner. Zero==none.	*/
-	FVector4	Normal;
+	FVector4f	Normal;
 	/** Time until hit.													*/
 	float		Time;
 	/** Primitive data item which was hit, INDEX_NONE=none				*/
@@ -148,7 +148,7 @@ public:
  * @param IntersectionTime	[in/out] Best intersection time so far (0..1), as in: IntersectionPoint = Start + IntersectionTime * Dir.
  * @return			true if the line intersected the triangle
  */
-FORCEINLINE bool appLineCheckTriangle(const FVector4& Start, const FVector4& End, const FVector4& Dir, const FVector4& V0, const FVector4& V1, const FVector4& V2, const FVector4& Normal, float& IntersectionTime)
+FORCEINLINE bool appLineCheckTriangle(const FVector4f& Start, const FVector4f& End, const FVector4f& Dir, const FVector4f& V0, const FVector4f& V1, const FVector4f& V2, const FVector4f& Normal, float& IntersectionTime)
 {
 	const float StartDist = FVectorUtils::PlaneDot(Normal, Start);
 	const float EndDist = FVectorUtils::PlaneDot(Normal, End);
@@ -173,8 +173,8 @@ FORCEINLINE bool appLineCheckTriangle(const FVector4& Start, const FVector4& End
 	}
 
 	// Calculate the line's point of intersection with the node's plane
-	const FVector4& Intersection = Start + Dir * Time;
-	const FVector4* Verts[3] = 
+	const FVector4f& Intersection = Start + Dir * Time;
+	const FVector4f* Verts[3] = 
 	{ 
 		&V0, &V1, &V2
 	};
@@ -182,7 +182,7 @@ FORCEINLINE bool appLineCheckTriangle(const FVector4& Start, const FVector4& End
 	// Check if the point of intersection is inside the triangle's edges.
 	for( int32 SideIndex = 0; SideIndex < 3; SideIndex++ )
 	{
-		const FVector4& SideDirection = Normal ^ (*Verts[(SideIndex + 1) % 3] - *Verts[SideIndex]);
+		const FVector4f& SideDirection = Normal ^ (*Verts[(SideIndex + 1) % 3] - *Verts[SideIndex]);
 		const float SideW = Dot3(SideDirection, *Verts[SideIndex]);
 		const float DotW = Dot3(SideDirection, Intersection);
 		if ((DotW - SideW) >= 0.001f)
@@ -391,15 +391,15 @@ struct FkDOPBuildCollisionTriangle
 	/**
 	 * First vertex in the triangle
 	 */
-	FVector4 V0;
+	FVector4f V0;
 	/**
 	 * Second vertex in the triangle
 	 */
-	FVector4 V1;
+	FVector4f V1;
 	/**
 	 * Third vertex in the triangle
 	 */
-	FVector4 V2;
+	FVector4f V2;
 
 #if CACHE_BUILD_TEMPORARIES
 private:
@@ -407,10 +407,10 @@ private:
  	 * Centroid of the triangle used for determining which bounding volume to
   	 * place the triangle in
 	 */
-	FVector4 Centroid;
+	FVector4f Centroid;
 
 	/** Cached local normal used by line/ triangle collision code. */
-	FVector4 LocalNormal;
+	FVector4f LocalNormal;
 public:
 #endif
 
@@ -433,22 +433,22 @@ public:
 	KDOP_IDX_TYPE MaterialIndex;
 
 #if CACHE_BUILD_TEMPORARIES
-	FORCEINLINE const FVector4& GetCentroid() const
+	FORCEINLINE const FVector4f& GetCentroid() const
 	{
 		return Centroid;
 	}
-	FORCEINLINE const FVector4& GetLocalNormal() const
+	FORCEINLINE const FVector4f& GetLocalNormal() const
 	{
 		return LocalNormal;
 	}
 #else
-	inline FVector4 GetCentroid() const
+	inline FVector4f GetCentroid() const
 	{
 		return (V0 + V1 + V2) / 3.f;
 	}
-	inline FVector4 GetLocalNormal() const
+	inline FVector4f GetLocalNormal() const
 	{
-		FVector4 LocalNormal = ((V1 - V2) ^ (V0 - V2)).GetSafeNormal();
+		FVector4f LocalNormal = ((V1 - V2) ^ (V0 - V2)).GetSafeNormal();
 		LocalNormal.W = Dot3(V0, LocalNormal);
 		return LocalNormal;
 	}
@@ -460,7 +460,7 @@ public:
 	 */
 	FkDOPBuildCollisionTriangle(
 		KDOP_IDX_TYPE InMaterialIndex,
-		const FVector4& vert0,const FVector4& vert1,const FVector4& vert2,
+		const FVector4f& vert0,const FVector4f& vert1,const FVector4f& vert2,
 		int32 InMeshIndex, 
 		uint32 InLODIndices,
 		uint32 InHLODRange,
@@ -498,12 +498,12 @@ struct FMultiBox
 	/**
 	 * Min planes for this set of bounding volumes. Array index is X/Y/Z.
 	 */
-	MS_ALIGN(16) FVector4 Min[3]  GCC_ALIGN(16);
+	MS_ALIGN(16) FVector4f Min[3]  GCC_ALIGN(16);
 
 	/** 
 	 * Max planes for this set of bounding volumes. Array index is X/Y/Z.
 	 */
-	FVector4 Max[3];
+	FVector4f Max[3];
 
 	/**
 	 * Sets the box at the passed in index to the passed in box.
@@ -511,7 +511,7 @@ struct FMultiBox
 	 * @param	BoundingVolumeIndex		Index of box to set
 	 * @param	Box						Box to set
 	 */
-	void SetBox( int32 BoundingVolumeIndex, const FBox& Box )
+	void SetBox( int32 BoundingVolumeIndex, const FBox3f& Box )
 	{
 		Min[0].Component(BoundingVolumeIndex) = Box.Min.X;
 		Min[1].Component(BoundingVolumeIndex) = Box.Min.Y;
@@ -528,11 +528,11 @@ struct FMultiBox
 	 *
 	 * @return Bounding volume at the passed in index
 	 */
-	FBox GetBox( int32 BoundingVolumeIndex )
+	FBox3f GetBox( int32 BoundingVolumeIndex )
 	{
-		FBox Box;
-		Box.Min = FVector4(Min[0][BoundingVolumeIndex],Min[1][BoundingVolumeIndex],Min[2][BoundingVolumeIndex],1);
-		Box.Max = FVector4(Max[0][BoundingVolumeIndex],Max[1][BoundingVolumeIndex],Max[2][BoundingVolumeIndex],1);
+		FBox3f Box;
+		Box.Min = FVector4f(Min[0][BoundingVolumeIndex],Min[1][BoundingVolumeIndex],Min[2][BoundingVolumeIndex],1);
+		Box.Max = FVector4f(Max[0][BoundingVolumeIndex],Max[1][BoundingVolumeIndex],Max[2][BoundingVolumeIndex],1);
 		return Box;
 	}
 };
@@ -548,7 +548,7 @@ namespace BoxTriangleIntersectionInternal
 
 struct FTriangle
 {
-	FVector Vertices[3];
+	FVector3f Vertices[3];
 };
 
 struct FOverlapInterval
@@ -557,14 +557,14 @@ struct FOverlapInterval
 	float Max;
 };
 
-FORCEINLINE FOverlapInterval GetInterval(const FTriangle& Triangle, const FVector& Vector)
+FORCEINLINE FOverlapInterval GetInterval(const FTriangle& Triangle, const FVector3f& Vector)
 {
 	FOverlapInterval Result;
-	Result.Min = Result.Max = FVector::DotProduct(Vector, Triangle.Vertices[0]);
+	Result.Min = Result.Max = FVector3f::DotProduct(Vector, Triangle.Vertices[0]);
 
 	for (int32 i = 1; i < 3; ++i)
 	{
-		float Projection = FVector::DotProduct(Vector, Triangle.Vertices[i]);
+		float Projection = FVector3f::DotProduct(Vector, Triangle.Vertices[i]);
 		Result.Min = FMath::Min(Result.Min, Projection);
 		Result.Max = FMath::Max(Result.Max, Projection);
 	}
@@ -572,26 +572,26 @@ FORCEINLINE FOverlapInterval GetInterval(const FTriangle& Triangle, const FVecto
 	return Result;
 }
 
-FORCEINLINE FOverlapInterval GetInterval(const FBox& Box, const FVector& Vector)
+FORCEINLINE FOverlapInterval GetInterval(const FBox3f& Box, const FVector3f& Vector)
 {
-	FVector BoxVertices[8] =
+	FVector3f BoxVertices[8] =
 	{
-		FVector(Box.Min.X, Box.Max.Y, Box.Max.Z),
-		FVector(Box.Min.X, Box.Max.Y, Box.Min.Z),
-		FVector(Box.Min.X, Box.Min.Y, Box.Max.Z),
-		FVector(Box.Min.X, Box.Min.Y, Box.Min.Z),
-		FVector(Box.Max.X, Box.Max.Y, Box.Max.Z),
-		FVector(Box.Max.X, Box.Max.Y, Box.Min.Z),
-		FVector(Box.Max.X, Box.Min.Y, Box.Max.Z),
-		FVector(Box.Max.X, Box.Min.Y, Box.Min.Z)
+		FVector3f(Box.Min.X, Box.Max.Y, Box.Max.Z),
+		FVector3f(Box.Min.X, Box.Max.Y, Box.Min.Z),
+		FVector3f(Box.Min.X, Box.Min.Y, Box.Max.Z),
+		FVector3f(Box.Min.X, Box.Min.Y, Box.Min.Z),
+		FVector3f(Box.Max.X, Box.Max.Y, Box.Max.Z),
+		FVector3f(Box.Max.X, Box.Max.Y, Box.Min.Z),
+		FVector3f(Box.Max.X, Box.Min.Y, Box.Max.Z),
+		FVector3f(Box.Max.X, Box.Min.Y, Box.Min.Z)
 	};
 
 	FOverlapInterval Result;
-	Result.Min = Result.Max = FVector::DotProduct(Vector, BoxVertices[0]);
+	Result.Min = Result.Max = FVector3f::DotProduct(Vector, BoxVertices[0]);
 
 	for (int32 i = 1; i < UE_ARRAY_COUNT(BoxVertices); ++i)
 	{
-		float Projection = FVector::DotProduct(Vector, BoxVertices[i]);
+		float Projection = FVector3f::DotProduct(Vector, BoxVertices[i]);
 		Result.Min = FMath::Min(Result.Min, Projection);
 		Result.Max = FMath::Max(Result.Max, Projection);
 	}
@@ -599,41 +599,41 @@ FORCEINLINE FOverlapInterval GetInterval(const FBox& Box, const FVector& Vector)
 	return Result;
 }
 
-FORCEINLINE bool OverlapOnAxis(const FBox& Box, const FTriangle& Triangle, const FVector& Vector)
+FORCEINLINE bool OverlapOnAxis(const FBox3f& Box, const FTriangle& Triangle, const FVector3f& Vector)
 {
 	FOverlapInterval A = GetInterval(Box, Vector);
 	FOverlapInterval B = GetInterval(Triangle, Vector);
 	return ((B.Min <= A.Max) && (A.Min <= B.Max));
 }
 
-FORCEINLINE bool IntersectTriangleAndAABB(const FTriangle& Triangle, const FBox& Box)
+FORCEINLINE bool IntersectTriangleAndAABB(const FTriangle& Triangle, const FBox3f& Box)
 {
-	FVector TriangleEdge0 = Triangle.Vertices[1] - Triangle.Vertices[0];
-	FVector TriangleEdge1 = Triangle.Vertices[2] - Triangle.Vertices[1];
-	FVector TriangleEdge2 = Triangle.Vertices[0] - Triangle.Vertices[2];
+	FVector3f TriangleEdge0 = Triangle.Vertices[1] - Triangle.Vertices[0];
+	FVector3f TriangleEdge1 = Triangle.Vertices[2] - Triangle.Vertices[1];
+	FVector3f TriangleEdge2 = Triangle.Vertices[0] - Triangle.Vertices[2];
 
-	FVector BoxNormal0(1.0f, 0.0f, 0.0f);
-	FVector BoxNormal1(0.0f, 1.0f, 0.0f);
-	FVector BoxNormal2(0.0f, 0.0f, 1.0f);
+	FVector3f BoxNormal0(1.0f, 0.0f, 0.0f);
+	FVector3f BoxNormal1(0.0f, 1.0f, 0.0f);
+	FVector3f BoxNormal2(0.0f, 0.0f, 1.0f);
 
-	FVector TestDirections[13] =
+	FVector3f TestDirections[13] =
 	{
 		// Separating axes from the box normals
 		BoxNormal0,
 		BoxNormal1,
 		BoxNormal2,
 		// One separating axis for the triangle normal
-		FVector::CrossProduct(TriangleEdge0, TriangleEdge1),
+		FVector3f::CrossProduct(TriangleEdge0, TriangleEdge1),
 		// Separating axes for the triangle edges
-		FVector::CrossProduct(BoxNormal0, TriangleEdge0),
-		FVector::CrossProduct(BoxNormal0, TriangleEdge1),
-		FVector::CrossProduct(BoxNormal0, TriangleEdge2),
-		FVector::CrossProduct(BoxNormal1, TriangleEdge0),
-		FVector::CrossProduct(BoxNormal1, TriangleEdge1),
-		FVector::CrossProduct(BoxNormal1, TriangleEdge2),
-		FVector::CrossProduct(BoxNormal2, TriangleEdge0),
-		FVector::CrossProduct(BoxNormal2, TriangleEdge1),
-		FVector::CrossProduct(BoxNormal2, TriangleEdge2)
+		FVector3f::CrossProduct(BoxNormal0, TriangleEdge0),
+		FVector3f::CrossProduct(BoxNormal0, TriangleEdge1),
+		FVector3f::CrossProduct(BoxNormal0, TriangleEdge2),
+		FVector3f::CrossProduct(BoxNormal1, TriangleEdge0),
+		FVector3f::CrossProduct(BoxNormal1, TriangleEdge1),
+		FVector3f::CrossProduct(BoxNormal1, TriangleEdge2),
+		FVector3f::CrossProduct(BoxNormal2, TriangleEdge0),
+		FVector3f::CrossProduct(BoxNormal2, TriangleEdge1),
+		FVector3f::CrossProduct(BoxNormal2, TriangleEdge2)
 	};
 
 	for (int i = 0; i < UE_ARRAY_COUNT(TestDirections); ++i)
@@ -719,23 +719,23 @@ struct TkDOPNode
 	 * @param	Check				Information about the ray to trace
 	 * @param	HitTime	[out]	Time of hit
 	 */
-	FORCEINLINE void LineCheckBounds(TkDOPLineCollisionCheck<COLL_DATA_PROVIDER,KDOP_IDX_TYPE>& Check, FVector4& HitTime, int32 NodeHit[4] ) const
+	FORCEINLINE void LineCheckBounds(TkDOPLineCollisionCheck<COLL_DATA_PROVIDER,KDOP_IDX_TYPE>& Check, FVector4f& HitTime, int32 NodeHit[4] ) const
 	{
 #define PLAIN_C 0
 #if PLAIN_C
 		for( int32 BoxIndex=0; BoxIndex<4; BoxIndex++ )
 		{
 			// 0: Create constants
-			FVector4 BoxMin( BoundingVolumes.Min[0][BoxIndex], BoundingVolumes.Min[1][BoxIndex], BoundingVolumes.Min[2][BoxIndex], 0 );
-			FVector4 BoxMax( BoundingVolumes.Max[0][BoxIndex], BoundingVolumes.Max[1][BoxIndex], BoundingVolumes.Max[2][BoxIndex], 0 );
+			FVector4f BoxMin( BoundingVolumes.Min[0][BoxIndex], BoundingVolumes.Min[1][BoxIndex], BoundingVolumes.Min[2][BoxIndex], 0 );
+			FVector4f BoxMax( BoundingVolumes.Max[0][BoxIndex], BoundingVolumes.Max[1][BoxIndex], BoundingVolumes.Max[2][BoxIndex], 0 );
 
 			// 1: Calculate slabs.
-			FVector4 Slab1 = (BoxMin - Check.LocalStart) * Check.LocalOneOverDir;
-			FVector4 Slab2 = (BoxMax - Check.LocalStart) * Check.LocalOneOverDir;
+			FVector4f Slab1 = (BoxMin - Check.LocalStart) * Check.LocalOneOverDir;
+			FVector4f Slab2 = (BoxMax - Check.LocalStart) * Check.LocalOneOverDir;
 
 			// 2: Figure out per component min/ max
-			FVector4 SlabMin = FVector4( Lightmass::Min(Slab1.X, Slab2.X), Lightmass::Min(Slab1.Y, Slab2.Y), Lightmass::Min(Slab1.Z, Slab2.Z), Lightmass::Min(Slab1.W, Slab2.W) );
-			FVector4 SlabMax = FVector4( Lightmass::Max(Slab1.X, Slab2.X), Lightmass::Max(Slab1.Y, Slab2.Y), Lightmass::Max(Slab1.Z, Slab2.Z), Lightmass::Max(Slab1.W, Slab2.W) );
+			FVector4f SlabMin = FVector4f( Lightmass::Min(Slab1.X, Slab2.X), Lightmass::Min(Slab1.Y, Slab2.Y), Lightmass::Min(Slab1.Z, Slab2.Z), Lightmass::Min(Slab1.W, Slab2.W) );
+			FVector4f SlabMax = FVector4f( Lightmass::Max(Slab1.X, Slab2.X), Lightmass::Max(Slab1.Y, Slab2.Y), Lightmass::Max(Slab1.Z, Slab2.Z), Lightmass::Max(Slab1.W, Slab2.W) );
 
 			// 3: Figure out global min/ max
 			float MinTime = Max3( SlabMin.X, SlabMin.Y, SlabMin.Z );
@@ -754,7 +754,7 @@ struct TkDOPNode
 		const LmVectorRegister InvDirY		= LmVectorSetFloat1( Check.LocalOneOverDir.Y );
 		const LmVectorRegister InvDirZ		= LmVectorSetFloat1( Check.LocalOneOverDir.Z );
 		const LmVectorRegister CurrentHitTime	= LmVectorSetFloat1( Check.Result->Time );
-		// Boxes are FVector2D so we need to unshuffle the data.
+		// Boxes are FVector2f so we need to unshuffle the data.
 		const LmVectorRegister BoxMinX		= LmVectorLoadAligned( &BoundingVolumes.Min[0] );
 		const LmVectorRegister BoxMinY		= LmVectorLoadAligned( &BoundingVolumes.Min[1] );
 		const LmVectorRegister BoxMinZ		= LmVectorLoadAligned( &BoundingVolumes.Min[2] );
@@ -809,7 +809,7 @@ struct TkDOPNode
 		{
 			SLOW_KDOP_STATS(FPlatformAtomics::InterlockedIncrement((SSIZE_T*)&GKDOPParentNodesTraversed));
 			// Check both left and right node at the same time.
-			FVector4 NodeHitTime;
+			FVector4f NodeHitTime;
 			MS_ALIGN(16) int32	NodeHit[4] GCC_ALIGN(16);
 			LineCheckBounds( Check, NodeHitTime, NodeHit );
 
@@ -881,7 +881,7 @@ struct TkDOPNode
 	 *
 	 * @param Check -- The aggregated line check data
 	 */
-	bool LineCheckPreCalculated(TkDOPLineCollisionCheck<COLL_DATA_PROVIDER,KDOP_IDX_TYPE>& Check, const FVector4& NodeHitTime, TTraversalHistory<KDOP_IDX_TYPE> History, int32* NodeHit) const
+	bool LineCheckPreCalculated(TkDOPLineCollisionCheck<COLL_DATA_PROVIDER,KDOP_IDX_TYPE>& Check, const FVector4f& NodeHitTime, TTraversalHistory<KDOP_IDX_TYPE> History, int32* NodeHit) const
 	{
 		bool bHit = 0;
 		// If this is a node, check the two child nodes and pick the closest one
@@ -1023,10 +1023,10 @@ struct TkDOPNode
 		for (int32 BoxIndex = 0; BoxIndex < 2; BoxIndex++)
 		{
 			// 0: Create constants
-			FVector BoxMin(BoundingVolumes.Min[0][BoxIndex], BoundingVolumes.Min[1][BoxIndex], BoundingVolumes.Min[2][BoxIndex]);
-			FVector BoxMax(BoundingVolumes.Max[0][BoxIndex], BoundingVolumes.Max[1][BoxIndex], BoundingVolumes.Max[2][BoxIndex]);
+			FVector3f BoxMin(BoundingVolumes.Min[0][BoxIndex], BoundingVolumes.Min[1][BoxIndex], BoundingVolumes.Min[2][BoxIndex]);
+			FVector3f BoxMax(BoundingVolumes.Max[0][BoxIndex], BoundingVolumes.Max[1][BoxIndex], BoundingVolumes.Max[2][BoxIndex]);
 
-			FBox ChildBox(BoxMin, BoxMax);
+			FBox3f ChildBox(BoxMin, BoxMax);
 
 			NodeHit[BoxIndex] = Check.Box.Intersect(ChildBox);
 		}
@@ -1146,7 +1146,7 @@ struct TkDOPTree
 	 * @param Nodes -- The list of nodes in this tree
 	 * @return bounding box for this node
 	 */
-	FBox SplitTriangleList(KDOP_IDX_TYPE NodeIndex, int32 Start, int32 NumTris, TArray<FkDOPBuildCollisionTriangle<KDOP_IDX_TYPE> >& BuildTriangles)
+	FBox3f SplitTriangleList(KDOP_IDX_TYPE NodeIndex, int32 Start, int32 NumTris, TArray<FkDOPBuildCollisionTriangle<KDOP_IDX_TYPE> >& BuildTriangles)
 	{
 		// This local node pointer will have to be updated whenever a Nodes reallocation can occur
 		NodeType* Node = &Nodes[NodeIndex];
@@ -1172,7 +1172,7 @@ struct TkDOPTree
 					// Project the centroid of the triangle against the plane
 					// normals and accumulate to find the total projected
 					// weighting
-					FVector4 Centroid = BuildTriangles[nTriangle].GetCentroid();
+					FVector4f Centroid = BuildTriangles[nTriangle].GetCentroid();
 
 					for (int32 nPlane = 0; nPlane < NUM_PLANES; nPlane++)
 					{
@@ -1190,7 +1190,7 @@ struct TkDOPTree
 				for (int32 nTriangle = Start; nTriangle < Start + NumTris; nTriangle++)
 				{
 					// Project the centroid again
-					FVector4 Centroid = BuildTriangles[nTriangle].GetCentroid();
+					FVector4f Centroid = BuildTriangles[nTriangle].GetCentroid();
 
 					// Now calculate the variance and accumulate it
 					for (int32 nPlane = 0; nPlane < NUM_PLANES; nPlane++)
@@ -1267,7 +1267,7 @@ struct TkDOPTree
 			Node->n.LeftNode = ChildIndex;
 			Node->n.RightNode = Node->n.LeftNode + 1;
 			// Have the left node recursively subdivide its list and set bounding volume.
-			FBox LeftBoundingVolume = SplitTriangleList(Node->n.LeftNode,Start,Left - Start,BuildTriangles);
+			FBox3f LeftBoundingVolume = SplitTriangleList(Node->n.LeftNode,Start,Left - Start,BuildTriangles);
 			// Nodes may have resized
 			Node = &Nodes[NodeIndex];
 			Node->BoundingVolumes.SetBox(0,LeftBoundingVolume);
@@ -1276,7 +1276,7 @@ struct TkDOPTree
 			Node->BoundingVolumes.SetBox(3,Nodes[Node->n.LeftNode].BoundingVolumes.GetBox(1));
 
 			// And now have the right node recursively subdivide its list and set bounding volume.			
-			FBox RightBoundingVolume = SplitTriangleList(Node->n.RightNode,Left,Start + NumTris - Left,BuildTriangles);
+			FBox3f RightBoundingVolume = SplitTriangleList(Node->n.RightNode,Left,Start + NumTris - Left,BuildTriangles);
 			// Nodes may have resized
 			Node = &Nodes[NodeIndex];
 			Node->BoundingVolumes.SetBox(1,RightBoundingVolume);
@@ -1292,7 +1292,7 @@ struct TkDOPTree
 
 			// "NULL triangle", used when a leaf can't fill all 4 triangles in a FTriangleSOA.
 			// No line should ever hit these triangles, set the values so that it can never happen.
-			FkDOPBuildCollisionTriangle<KDOP_IDX_TYPE> EmptyTriangle(0,FVector4(0,0,0,0),FVector4(0,0,0,0),FVector4(0,0,0,0),INDEX_NONE,INDEX_NONE,INDEX_NONE, false, true);
+			FkDOPBuildCollisionTriangle<KDOP_IDX_TYPE> EmptyTriangle(0,FVector4f(0,0,0,0),FVector4f(0,0,0,0),FVector4f(0,0,0,0),INDEX_NONE,INDEX_NONE,INDEX_NONE, false, true);
 			
 			Node->t.StartIndex = SOATriangles.Num();
 			Node->t.NumTriangles = Align<int32>(NumTris, 4) / 4;
@@ -1324,10 +1324,10 @@ struct TkDOPTree
 				SOA.Positions[2].Y = LmVectorSet( Tris[0]->V2.Y, Tris[1]->V2.Y, Tris[2]->V2.Y, Tris[3]->V2.Y );
 				SOA.Positions[2].Z = LmVectorSet( Tris[0]->V2.Z, Tris[1]->V2.Z, Tris[2]->V2.Z, Tris[3]->V2.Z );
 
-				const FVector4& Tris0LocalNormal = Tris[0]->GetLocalNormal();
-				const FVector4& Tris1LocalNormal = Tris[1]->GetLocalNormal();
-				const FVector4& Tris2LocalNormal = Tris[2]->GetLocalNormal();
-				const FVector4& Tris3LocalNormal = Tris[3]->GetLocalNormal();
+				const FVector4f& Tris0LocalNormal = Tris[0]->GetLocalNormal();
+				const FVector4f& Tris1LocalNormal = Tris[1]->GetLocalNormal();
+				const FVector4f& Tris2LocalNormal = Tris[2]->GetLocalNormal();
+				const FVector4f& Tris3LocalNormal = Tris[3]->GetLocalNormal();
 
 				SOA.Normals.X = LmVectorSet( Tris0LocalNormal.X, Tris1LocalNormal.X, Tris2LocalNormal.X, Tris3LocalNormal.X );
 				SOA.Normals.Y = LmVectorSet( Tris0LocalNormal.Y, Tris1LocalNormal.Y, Tris2LocalNormal.Y, Tris3LocalNormal.Y );
@@ -1353,7 +1353,7 @@ struct TkDOPTree
 			Node->Occupancy = NumTris;
 			
 			// Generate bounding volume for leaf which is passed up the call chain.
-			FBox BoundingVolume(ForceInit);
+			FBox3f BoundingVolume(ForceInit);
 			for (int32 TriangleIndex=Start; TriangleIndex<Start + NumTris; TriangleIndex++)
 			{
 				BoundingVolume += BuildTriangles[TriangleIndex].V0;
@@ -1473,8 +1473,8 @@ template <typename COLL_DATA_PROVIDER, typename KDOP_IDX_TYPE> struct TkDOPLineC
 	 */
 	FHitResult* Result;
 	// Constant input vars
-	const FVector4& Start;
-	const FVector4& End;
+	const FVector4f& Start;
+	const FVector4f& End;
 
 	/**
 	 * Flags for optimizing a trace
@@ -1484,12 +1484,12 @@ template <typename COLL_DATA_PROVIDER, typename KDOP_IDX_TYPE> struct TkDOPLineC
 	const bool bTwoSidedCollision;
 	const bool bFlipSidedness;
 	// Locally calculated vectors
-	FVector4 LocalStart;
-	FVector4 LocalEnd;
-	FVector4 LocalDir;
-	FVector4 LocalOneOverDir;
+	FVector4f LocalStart;
+	FVector4f LocalEnd;
+	FVector4f LocalDir;
+	FVector4f LocalOneOverDir;
 	// Normal in local space which gets transformed to world at the very end
-	FVector4 LocalHitNormal;
+	FVector4f LocalHitNormal;
 
 	/** Index into the kDOP tree's nodes of the node that was hit. */
 	KDOP_IDX_TYPE HitNodeIndex;
@@ -1519,7 +1519,7 @@ template <typename COLL_DATA_PROVIDER, typename KDOP_IDX_TYPE> struct TkDOPLineC
 	 *		specific data, such as L2W, W2L, Vertices, and so on
 	 * @param InResult -- The out param for hit result information
 	 */
-	TkDOPLineCollisionCheck(const FVector4& InStart,const FVector4& InEnd,
+	TkDOPLineCollisionCheck(const FVector4f& InStart,const FVector4f& InEnd,
 		bool bInbFindClosestIntersection,
 		bool bInStaticAndOpaqueOnly,
 		bool bInTwoSidedCollision,
@@ -1540,7 +1540,7 @@ template <typename COLL_DATA_PROVIDER, typename KDOP_IDX_TYPE> struct TkDOPLineC
 		bFlipSidedness(bInFlipSidedness),
 		HitNodeIndex(0xFFFFFFFF)
 	{
-		const FMatrix& WorldToLocal = TkDOPCollisionCheck<COLL_DATA_PROVIDER,KDOP_IDX_TYPE>::CollDataProvider.GetWorldToLocal();
+		const FMatrix44f& WorldToLocal = TkDOPCollisionCheck<COLL_DATA_PROVIDER,KDOP_IDX_TYPE>::CollDataProvider.GetWorldToLocal();
 		// Move start and end to local space
 		LocalStart = WorldToLocal.TransformPosition(Start);
 		LocalEnd = WorldToLocal.TransformPosition(End);
@@ -1570,10 +1570,10 @@ template <typename COLL_DATA_PROVIDER, typename KDOP_IDX_TYPE> struct TkDOPLineC
 	 * Transforms the local hit normal into a world space normal using the transpose
 	 * adjoint and flips the normal if need be
 	 */
-	FORCEINLINE FVector4 GetHitNormal(void)
+	FORCEINLINE FVector4f GetHitNormal(void)
 	{
 		// Transform the hit back into world space using the transpose adjoint
-		FVector4 Normal = TkDOPCollisionCheck<COLL_DATA_PROVIDER,KDOP_IDX_TYPE>::CollDataProvider.GetLocalToWorldTransposeAdjoint().TransformVector(LocalHitNormal).GetSafeNormal();
+		FVector4f Normal = TkDOPCollisionCheck<COLL_DATA_PROVIDER,KDOP_IDX_TYPE>::CollDataProvider.GetLocalToWorldTransposeAdjoint().TransformVector(LocalHitNormal).GetSafeNormal();
 		// Flip the normal if the triangle is inverted
 		if (TkDOPCollisionCheck<COLL_DATA_PROVIDER,KDOP_IDX_TYPE>::CollDataProvider.GetDeterminant() < 0.f)
 		{
@@ -1586,11 +1586,11 @@ template <typename COLL_DATA_PROVIDER, typename KDOP_IDX_TYPE> struct TkDOPLineC
 template <typename COLL_DATA_PROVIDER, typename KDOP_IDX_TYPE> struct TkDOPBoxCollisionCheck :
 	public TkDOPCollisionCheck<COLL_DATA_PROVIDER, KDOP_IDX_TYPE>
 {
-	const FBox Box;
+	const FBox3f Box;
 	int32 Item;
 
 	TkDOPBoxCollisionCheck(
-		const FBox& InBox,
+		const FBox3f& InBox,
 		const COLL_DATA_PROVIDER& InCollDataProvider) :
 		TkDOPCollisionCheck<COLL_DATA_PROVIDER, KDOP_IDX_TYPE>(InCollDataProvider),
 		Box(InBox),
diff --git a/Engine/Source/Programs/UnrealLightmass/Public/ImportExport.h b/Engine/Source/Programs/UnrealLightmass/Public/ImportExport.h
--- a/Engine/Source/Programs/UnrealLightmass/Public/ImportExport.h
+++ b/Engine/Source/Programs/UnrealLightmass/Public/ImportExport.h
@@ -75,17 +75,17 @@ namespace Lightmass
 	 * This must be done whenever modifying code that changes how each data type is exported!
 	 */
 
-	static const int32 LM_TEXTUREMAPPING_VERSION		= 1;
-	static const int32 LM_VOLUMESAMPLES_VERSION			= 1;
-	static const int32 LM_VOLUMETRICLIGHTMAP_VERSION		= 3;
-	static const int32 LM_PRECOMPUTEDVISIBILITY_VERSION	= 1;
-	static const int32 LM_VOLUMEDEBUGOUTPUT_VERSION		= 1;
-	static const int32 LM_DOMINANTSHADOW_VERSION		= 1;
-	static const int32 LM_MESHAREALIGHTDATA_VERSION		= 1;
-	static const int32 LM_DEBUGOUTPUT_VERSION			= 1;
-	static const int32 LM_SCENE_VERSION					= 1;
-	static const int32 LM_STATICMESH_VERSION			= 1;
-	static const int32 LM_MATERIAL_VERSION				= 1;
+	static const int32 LM_TEXTUREMAPPING_VERSION		= 2;
+	static const int32 LM_VOLUMESAMPLES_VERSION			= 2;
+	static const int32 LM_VOLUMETRICLIGHTMAP_VERSION		= 4;
+	static const int32 LM_PRECOMPUTEDVISIBILITY_VERSION	= 2;
+	static const int32 LM_VOLUMEDEBUGOUTPUT_VERSION		= 2;
+	static const int32 LM_DOMINANTSHADOW_VERSION		= 2;
+	static const int32 LM_MESHAREALIGHTDATA_VERSION		= 2;
+	static const int32 LM_DEBUGOUTPUT_VERSION			= 2;
+	static const int32 LM_SCENE_VERSION					= 2;
+	static const int32 LM_STATICMESH_VERSION			= 2;
+	static const int32 LM_MATERIAL_VERSION				= 2;
 
 
 	/** Alert source object type identifiers... */
@@ -288,7 +288,7 @@ namespace Lightmass
 	{
 	public:
 		/** World space position and radius. */
-		FVector4 PositionAndRadius;
+		FVector4f PositionAndRadius;
 
 		/** SH coefficients used with high quality lightmaps. */
 		float HighQualityCoefficients[LM_NUM_SH_COEFFICIENTS][3];
@@ -296,7 +296,7 @@ namespace Lightmass
 		/** SH coefficients used with low quality lightmaps. */
 		float LowQualityCoefficients[LM_NUM_SH_COEFFICIENTS][3];
 
-		FVector SkyBentNormal;
+		FVector3f SkyBentNormal;
 
 		/** Shadow factor for the stationary directional light. */
 		float DirectionalLightShadowing;
@@ -321,7 +321,7 @@ namespace Lightmass
 	{
 	public:
 		/** Transform from world space to the coordinate space that FStaticShadowDepthMapSampleData's are stored in. */
-		FMatrix WorldToLight;
+		FMatrix44f WorldToLight;
 		/** Dimensions of the generated shadow map. */
 		int32 ShadowMapSizeX;
 		int32 ShadowMapSizeY;
@@ -335,8 +335,8 @@ namespace Lightmass
 	{
 	public:
 		FGuid LevelGuid;
-		FVector4 Position;
-		FVector4 Direction;
+		FVector4f Position;
+		FVector4f Direction;
 		float Radius;
 		float ConeAngle;
 		FColor Color;
@@ -351,13 +351,13 @@ namespace Lightmass
 
 	struct FDebugStaticLightingRay
 	{
-		FVector4 Start;
-		FVector4 End;
+		FVector4f Start;
+		FVector4f End;
 		bool bHit;
 		bool bPositive;
 
 		FDebugStaticLightingRay() {}
-		FDebugStaticLightingRay(const FVector4& InStart, const FVector4& InEnd, bool bInHit, bool bInPositive = false) :
+		FDebugStaticLightingRay(const FVector4f& InStart, const FVector4f& InEnd, bool bInHit, bool bInPositive = false) :
 			Start(InStart),
 			End(InEnd),
 			bHit(bInHit),
@@ -367,8 +367,8 @@ namespace Lightmass
 
 	struct FDebugStaticLightingVertex
 	{
-		FVector4 VertexNormal;
-		FVector4 VertexPosition;
+		FVector4f VertexNormal;
+		FVector4f VertexPosition;
 	};
 
 	struct FDebugLightingCacheRecord
@@ -388,15 +388,15 @@ namespace Lightmass
 	struct FDebugPhoton
 	{
 		int32 Id;
-		FVector4 Position;
-		FVector4 Direction;
-		FVector4 Normal;
+		FVector4f Position;
+		FVector4f Direction;
+		FVector4f Normal;
 
 		FDebugPhoton() :
 			Id(-1)
 		{}
 
-		FDebugPhoton(int32 InId, const FVector4& InPosition, const FVector4& InDirection, const FVector4& InNormal) :
+		FDebugPhoton(int32 InId, const FVector4f& InPosition, const FVector4f& InDirection, const FVector4f& InNormal) :
 			Id(InId),
 			Position(InPosition),
 			Direction(InDirection),
@@ -406,10 +406,10 @@ namespace Lightmass
 
 	struct FDebugOctreeNode
 	{
-		FVector4 Center;
-		FVector4 Extent;
+		FVector4f Center;
+		FVector4f Extent;
 
-		FDebugOctreeNode(const FVector4& InCenter, const FVector4& InExtent) :
+		FDebugOctreeNode(const FVector4f& InCenter, const FVector4f& InExtent) :
 			Center(InCenter),
 			Extent(InExtent)
 		{}
@@ -454,7 +454,7 @@ namespace Lightmass
 		/** Direct photon that was found during direct lighting. */
 		FDebugPhoton GatheredDirectPhoton;
 		/** Positions of the selected texel's corners */
-		FVector4 TexelCorners[NumTexelCorners];
+		FVector4f TexelCorners[NumTexelCorners];
 		/** Whether each of the selected texel's corners were valid */
 		bool bCornerValid[NumTexelCorners];
 		/** World space radius of the selected sample */
@@ -473,10 +473,10 @@ namespace Lightmass
 
 	struct FDebugVolumeLightingSample
 	{
-		FVector4 PositionAndRadius;
+		FVector4f PositionAndRadius;
 		FLinearColor AverageIncidentRadiance;
 
-		FDebugVolumeLightingSample(const FVector4& InPositionAndRadius, const FLinearColor& InAverageIncidentRadiance) :
+		FDebugVolumeLightingSample(const FVector4f& InPositionAndRadius, const FLinearColor& InAverageIncidentRadiance) :
 			PositionAndRadius(InPositionAndRadius),
 			AverageIncidentRadiance(InAverageIncidentRadiance)
 		{}
diff --git a/Engine/Source/Programs/UnrealLightmass/Public/MeshExport.h b/Engine/Source/Programs/UnrealLightmass/Public/MeshExport.h
--- a/Engine/Source/Programs/UnrealLightmass/Public/MeshExport.h
+++ b/Engine/Source/Programs/UnrealLightmass/Public/MeshExport.h
@@ -97,11 +97,11 @@ struct FStaticMeshElementData
 //----------------------------------------------------------------------------
 struct FStaticMeshVertex
 {
-	FVector4		Position;
-	FVector4		TangentX;
-	FVector4		TangentY;
-	FVector4		TangentZ;
-	FVector2D		UVs[MAX_TEXCOORDS];
+	FVector4f		Position;
+	FVector4f		TangentX;
+	FVector4f		TangentY;
+	FVector4f		TangentZ;
+	FVector2f		UVs[MAX_TEXCOORDS];
 };
 
 #if !PLATFORM_MAC && !PLATFORM_LINUX
