Fix CPU Lightmass by coverting LWC types to single precision

diff --git a/Engine/Source/Programs/UnrealLightmass/Public/SceneExport.h b/Engine/Source/Programs/UnrealLightmass/Public/SceneExport.h
--- a/Engine/Source/Programs/UnrealLightmass/Public/SceneExport.h
+++ b/Engine/Source/Programs/UnrealLightmass/Public/SceneExport.h
@@ -340,8 +340,8 @@ public:
 	FIntVector TopLevelGridSize;
 
 	/** World space size of the volumetric lightmap. */
-	FVector VolumeMin;
-	FVector VolumeSize;
+	FVector3f VolumeMin;
+	FVector3f VolumeSize;
 
 	/** 
 	 * Size of a brick of unique lighting data.  Must be a power of 2.  
@@ -391,7 +391,7 @@ public:
 	bool bPlaceCellsOnlyAlongCameraTracks;
 
 	/** World space size of visibility cells in the x and y dimensions. */
-	FVector4::FReal CellSize;
+	float CellSize;
 
 	/** Number of tasks that visibility cells are being split up into. */
 	int32 NumCellDistributionBuckets;
@@ -785,7 +785,7 @@ struct FDebugLightingInputData
 	/** Index of the BSP node to debug if the mapping is a BSP mapping */
 	int32 NodeIndex;
 	/** World space position of the position that was clicked to select the debug sample */
-	FVector4 Position;
+	FVector4f Position;
 	/** Position in the texture mapping of the texel to debug */
 	int32 LocalX;
 	int32 LocalY;
@@ -793,7 +793,7 @@ struct FDebugLightingInputData
 	int32 MappingSizeX;
 	int32 MappingSizeY;
 	/** Position of the camera */
-	FVector4 CameraPosition;
+	FVector4f CameraPosition;
 	/** VisibilityId of a component from the currently selected actor or BSP surface. */
 	int32 DebugVisibilityId;
 };
@@ -897,8 +897,8 @@ struct FLightData
 	/** Bit-wise combination of flags from EDawnLightFlags */
 	uint32			LightFlags;
 	/** Homogeneous coordinates */
-	FVector4		Position;
-	FVector4		Direction;
+	FVector4f		Position;
+	FVector4f		Direction;
 	FLinearColor	Color;
 	float			Brightness;
 	/** The radius of the light's surface, not the light's influence. */
@@ -918,21 +918,21 @@ struct FLightData
 	// @param LightPosition		position of the light in the world
 	// @param LightDirection	direction at which the light is emitting (x axis)
 	// @param LightTangent		tangent to the direction at which the light is emitting (z axis)
-	inline float ComputeLightProfileMultiplier(const TArray< uint8 >& LightProfileTextureData, FVector WorldPosition, FVector LightPosition, FVector LightDirection, FVector LightTangent) const
+	inline float ComputeLightProfileMultiplier(const TArray< uint8 >& LightProfileTextureData, FVector3f WorldPosition, FVector3f LightPosition, FVector3f LightDirection, FVector3f LightTangent) const
 	{
 		// optimization - only evaluate this function if needed
 		if(LightFlags & Lightmass::GI_LIGHT_USE_LIGHTPROFILE)
 		{
-			FVector LightBitangent = FVector::CrossProduct( LightTangent, LightDirection ).GetSafeNormal();
+			FVector3f LightBitangent = FVector3f::CrossProduct( LightTangent, LightDirection ).GetSafeNormal();
 
-			FMatrix LightTransform = FMatrix( LightDirection, LightBitangent, LightTangent, FVector4(0.f, 0.f, 0.f, 1.f) );
-			FMatrix InvLightTransform = LightTransform.GetTransposed();
+			FMatrix44f LightTransform = FMatrix44f( LightDirection, LightBitangent, LightTangent, FVector4f(0.f, 0.f, 0.f, 1.f) );
+			FMatrix44f InvLightTransform = LightTransform.GetTransposed();
 
-			FVector ToLight = (LightPosition - WorldPosition).GetSafeNormal();
-			FVector LocalToLight = InvLightTransform.TransformVector( ToLight );
+			FVector3f ToLight = (LightPosition - WorldPosition).GetSafeNormal();
+			FVector3f LocalToLight = InvLightTransform.TransformVector( ToLight );
 
 			// -1..1
-			float DotProd = FVector::DotProduct(ToLight, LightDirection);
+			float DotProd = FVector3f::DotProduct(ToLight, LightDirection);
 			// -PI..PI (this distortion could be put into the texture but not without quality loss or more memory)
 			float Angle = FMath::Asin(DotProd);
 			// 0..1
@@ -999,7 +999,7 @@ struct FPointLightData
 	float		Radius;
 	float		FalloffExponent;
 	// Point lights need an additional axis to specify the direction of IES profiles, also used by spot lights for the direction of tube lights
-	FVector		LightTangent;
+	FVector3f		LightTangent;
 };
 
 //----------------------------------------------------------------------------
@@ -1119,7 +1119,7 @@ struct FStaticLightingMeshInstanceData
 	/** The lights which affect the mesh's primitive. */
 	int32 NumRelevantLights;
 	/** The bounding box of the mesh. */
-	FBox BoundingBox;
+	FBox3f BoundingBox;
 };
 
 namespace ESplineMeshAxis
@@ -1139,27 +1139,27 @@ namespace ESplineMeshAxis
 struct FSplineMeshParams
 {
 	/** Start location of spline, in component space */
-	FVector StartPos;
+	FVector3f StartPos;
 	/** Start tangent of spline, in component space */
-	FVector StartTangent;
+	FVector3f StartTangent;
 	/** X and Y scale applied to mesh at start of spline */
-	FVector2D StartScale;
+	FVector2f StartScale;
 	/** Roll around spline applied at start */
 	float StartRoll;
 	/** Offset from the spline at start */
-	FVector2D StartOffset;
+	FVector2f StartOffset;
 	/** End location of spline, in component space */
-	FVector EndPos;
+	FVector3f EndPos;
 	/** End tangent of spline, in component space */
-	FVector EndTangent;
+	FVector3f EndTangent;
 	/** X and Y scale applied to mesh at end of spline */
-	FVector2D EndScale;
+	FVector2f EndScale;
 	/** Roll around spline applied at end */
 	float EndRoll;
 	/** Offset from the base spline at end */
-	FVector2D EndOffset;
+	FVector2f EndOffset;
 	/** Axis (in component space) that is used to determine X axis for co-ordinates along spline */
-	FVector SplineUpDir;
+	FVector3f SplineUpDir;
 	/** Smoothly (cubic) interpolate the Roll and Scale params over spline. */
 	bool bSmoothInterpRollScale;
 	/** Minimum Z value of the entire mesh */
@@ -1175,7 +1175,7 @@ struct FStaticMeshStaticLightingMeshData
 	/** The LOD this mesh represents. */
 	uint32 EncodedLODIndices;
 	uint32 EncodedHLODRange;
-	FMatrix LocalToWorld;
+	FMatrix44f LocalToWorld;
 	/** true if the primitive has a transform which reverses the winding of its triangles. */
 	bool bReverseWinding;
 	bool bShouldSelfShadow;
@@ -1186,25 +1186,25 @@ struct FStaticMeshStaticLightingMeshData
 
 struct FMinimalStaticLightingVertex
 {
-	FVector4 WorldPosition;
-	FVector4 WorldTangentZ;
+	FVector4f WorldPosition;
+	FVector4f WorldTangentZ;
 	FVector2f TextureCoordinates[MAX_TEXCOORDS];
 };
 
 struct FStaticLightingVertexData : public FMinimalStaticLightingVertex
 {
-	FVector4 WorldTangentX;
-	FVector4 WorldTangentY;
+	FVector4f WorldTangentX;
+	FVector4f WorldTangentY;
 };
 
 struct FBSPSurfaceStaticLightingData
 {
-	FVector4	TangentX;
-	FVector4	TangentY;
-	FVector4	TangentZ;
+	FVector4f	TangentX;
+	FVector4f	TangentY;
+	FVector4f	TangentZ;
 
-	FMatrix		MapToWorld;
-	FMatrix		WorldToMap;
+	FMatrix44f		MapToWorld;
+	FMatrix44f		WorldToMap;
 
 	FGuid		MaterialGuid;
 };
@@ -1242,12 +1242,12 @@ struct FStaticLightingVertexMappingData
 struct FFluidSurfaceStaticLightingMeshData
 {
 	/** The primitive's local to world transform. */
-	FMatrix LocalToWorld;
+	FMatrix44f LocalToWorld;
 	/** The inverse transpose of the primitive's local to world transform. */
-	FMatrix LocalToWorldInverseTranspose;
+	FMatrix44f LocalToWorldInverseTranspose;
 	/** The mesh data of the fluid surface, which is represented as a quad. */
-	FVector4 QuadCorners[4];
-	FVector4 QuadUVCorners[4];
+	FVector4f QuadCorners[4];
+	FVector4f QuadUVCorners[4];
 	int32 QuadIndices[6];
 };
 
@@ -1257,7 +1257,7 @@ struct FFluidSurfaceStaticLightingMeshData
 struct FLandscapeStaticLightingMeshData
 {
 	/** The primitive's local to world transform. */
-	FMatrix LocalToWorld;
+	FMatrix44f LocalToWorld;
 	int32 ComponentSizeQuads;
 	float LightMapRatio;
 	/** The number of quads we are expanding to eliminate seams. */
@@ -1267,7 +1267,7 @@ struct FLandscapeStaticLightingMeshData
 
 struct FVolumetricLightmapDensityVolumeData
 {
-	FBox Bounds;
+	FBox3f Bounds;
 	FIntPoint AllowedMipLevelRange;
 	int32 NumPlanes;
 };
