Fix CPU Lightmass Mac/Linux compilation

diff --git a/Engine/Source/Programs/UnrealLightmass/Private/ImportExport/LightmassScene.cpp b/Engine/Source/Programs/UnrealLightmass/Private/ImportExport/LightmassScene.cpp
--- a/Engine/Source/Programs/UnrealLightmass/Private/ImportExport/LightmassScene.cpp
+++ b/Engine/Source/Programs/UnrealLightmass/Private/ImportExport/LightmassScene.cpp
@@ -186,13 +186,7 @@ void FScene::Import( FLightmassImporter& Importer )
 	{
 		FVolumetricLightmapDensityVolumeData LMVolumeData;
 		Importer.ImportData(&LMVolumeData);
-		// LWC_TODO: SceneExport.h does #pragma pack(1) for non-mac/linux platforms meaning there's a 4 byte padding difference depending on platform when large world coordinates are enabled. Would be good to make this consistent.
-#if (!UE_LARGE_WORLD_COORDINATES_DISABLED && (PLATFORM_MAC || PLATFORM_LINUX))
-		static_assert(sizeof(LMVolumeData) == sizeof(FBox3f) + sizeof(FIntPoint) + sizeof(int32) + 4, "Update member copy");
-#else
 		static_assert(sizeof(LMVolumeData) == sizeof(FBox3f) + sizeof(FIntPoint) + sizeof(int32), "Update member copy");
-#endif
-
 		FVolumetricLightmapDensityVolume LMVolume;
 		LMVolume.Bounds = LMVolumeData.Bounds;
 		LMVolume.AllowedMipLevelRange = LMVolumeData.AllowedMipLevelRange;
