Fix crash in UnrealLightmass on exit with stats enabled. Repro'd via:

  1. Open QAGame in UE
  2. Open StaticLightingFromMat level
  3. In the Content Browser, search for 'M_Emissive_Grtr' and place it on the sphere mesh in the map
  4. In the details of Panel for the Sphere Mesh, search for 'Use Emissive for Static Lighting' and check the box if it is not already checked
  5. Build lighting in Production

We were doing this at the end of LightmassMain:

  FTaskGraphInterface::Shutdown();
  FEngineLoop::AppExit();

FTaskGraphInterface was being used in FThreadStats::StopThread() and we'd crash.

FEngineLoop::AppExit() does the thread stats shutdown before calling FTaskGraphInterface::Shutdown. Ie:

  void FEngineLoop::AppExit()
  {
    // when compiled WITH_ENGINE, this will happen in FEngineLoop::Exit()
#if !WITH_ENGINE
#if STATS
    FThreadStats::StopThread();
#endif
    FTaskGraphInterface::Shutdown();
#endif // WITH_ENGINE

This will stop the stats thread (if needed) before shutting down FTaskGraphInterface.

diff --git a/Engine/Source/Programs/UnrealLightmass/Private/Launch/UnrealLightmass.cpp b/Engine/Source/Programs/UnrealLightmass/Private/Launch/UnrealLightmass.cpp
--- a/Engine/Source/Programs/UnrealLightmass/Private/Launch/UnrealLightmass.cpp
+++ b/Engine/Source/Programs/UnrealLightmass/Private/Launch/UnrealLightmass.cpp
@@ -250,7 +250,10 @@ int LightmassMain(int argc, ANSICHAR* argv[])
 	FEngineLoop::AppPreExit();
 	FModuleManager::Get().UnloadModulesAtShutdown();
 
-	FTaskGraphInterface::Shutdown();
+	// Similar to CL 16324633. If we shut down FTaskGraphInterface right here, in AppExit
+	//   FThreadStats::StopThread could get called and that will try and execute some stats
+	//   commands which use FTaskGraphInterface and --> crashes. FEngineLoop::AppExit() calls
+	//   FTaskGraphInterface::Shutdown() after calling FThreadStats::StopThread() if needed.
 	FEngineLoop::AppExit();
 #endif
 
