Fixed compilation error in lightmass as a fallout of 18751249

diff --git a/Engine/Source/Programs/UnrealLightmass/Private/Lighting/LightingSystem.inl b/Engine/Source/Programs/UnrealLightmass/Private/Lighting/LightingSystem.inl
--- a/Engine/Source/Programs/UnrealLightmass/Private/Lighting/LightingSystem.inl
+++ b/Engine/Source/Programs/UnrealLightmass/Private/Lighting/LightingSystem.inl
@@ -26,7 +26,7 @@ TGatheredLightSample<SHOrder> FGatheredLightSampleUtil::PointLightWorldSpace(con
 
 	if (TangentDirection.Z >= 0.0f)
 	{
-		Result.SHVector.AddIncomingRadiance(Color, 1, WorldDirection);
+		Result.SHVector.AddIncomingRadiance(Color, 1, (const FVector4&)WorldDirection);
 
 		FSHVector2 SH = FSHVector2::SHBasisFunction(FVector4(TangentDirection));
 		// Evaluate lighting along the smoothed vertex normal direction, so that later we can guarantee an SH intensity of 1 along the normal
