#ue5 - Fix lightmass after Matrix changes

diff --git a/Engine/Source/Programs/UnrealLightmass/Private/Lighting/LightingSystem.cpp b/Engine/Source/Programs/UnrealLightmass/Private/Lighting/LightingSystem.cpp
--- a/Engine/Source/Programs/UnrealLightmass/Private/Lighting/LightingSystem.cpp
+++ b/Engine/Source/Programs/UnrealLightmass/Private/Lighting/LightingSystem.cpp
@@ -2180,8 +2180,8 @@ void FStaticLightingSystem::CalculateStaticShadowDepthMap(FGuid LightGuid)
 			}
 		}
 
-		ShadowDepthMap->WorldToLight *= FTranslationMatrix(-LightSpaceImportanceBounds.Min)
-			* FScaleMatrix(FVector3f(1.0f) / (LightSpaceImportanceBounds.Max - LightSpaceImportanceBounds.Min));
+		ShadowDepthMap->WorldToLight *= FTranslationMatrix44f(-LightSpaceImportanceBounds.Min)
+			* FScaleMatrix44f(FVector3f(1.0f) / (LightSpaceImportanceBounds.Max - LightSpaceImportanceBounds.Min));
 
 		FScopeLock Lock(&CompletedStaticShadowDepthMapsSync);
 		CompletedStaticShadowDepthMaps.Add(DirectionalLight, ShadowDepthMap);
@@ -2274,7 +2274,7 @@ void FStaticLightingSystem::CalculateStaticShadowDepthMap(FGuid LightGuid)
 
 		ShadowDepthMap->WorldToLight *= 
 			// Perspective projection sized to the spotlight cone
-			FPerspectiveMatrix(SpotLight->OuterConeAngle * (float)PI / 180.0f, 1, 1, 0, SpotLight->Radius)
+			FPerspectiveMatrix44f(SpotLight->OuterConeAngle * (float)PI / 180.0f, 1, 1, 0, SpotLight->Radius)
 			// Convert from NDC to texture space, normalize Z
 			* FMatrix44f(
 				FPlane4f(.5f,								0,							0,									0),
