Add an ON_SCOPE_EXIT for LightmassMain cleanup with local swarm interface

When running UnrealLightmass with -help and various other arguments we were just doing "returns" without any cleanup.

This would lead to various crashes in destructors, etc.

diff --git a/Engine/Source/Programs/UnrealLightmass/Private/Launch/UnrealLightmass.cpp b/Engine/Source/Programs/UnrealLightmass/Private/Launch/UnrealLightmass.cpp
--- a/Engine/Source/Programs/UnrealLightmass/Private/Launch/UnrealLightmass.cpp
+++ b/Engine/Source/Programs/UnrealLightmass/Private/Launch/UnrealLightmass.cpp
@@ -98,7 +98,19 @@ int LightmassMain(int argc, ANSICHAR* argv[])
 	FModuleManager::LoadModuleChecked<IMessagingModule>("Messaging");
 	FModuleManager::Get().LoadModule(TEXT("Settings"));
 	IPluginManager::Get().LoadModulesForEnabledPlugins(ELoadingPhase::PreDefault);
-#endif
+
+	ON_SCOPE_EXIT
+	{
+		FEngineLoop::AppPreExit();
+		FModuleManager::Get().UnloadModulesAtShutdown();
+
+		// Similar to CL 16324633. If we shut down FTaskGraphInterface right here, in AppExit
+		//   FThreadStats::StopThread could get called and that will try and execute some stats
+		//   commands which use FTaskGraphInterface and --> crashes. FEngineLoop::AppExit() calls
+		//   FTaskGraphInterface::Shutdown() after calling FThreadStats::StopThread() if needed.
+		FEngineLoop::AppExit();
+	};
+#endif // USE_LOCAL_SWARM_INTERFACE
 
 	UE_LOG(LogLightmass, Display,  TEXT("Lightmass %s started on: %s. Command-line: %s"), FPlatformMisc::GetUBTPlatform(), FPlatformProcess::ComputerName(), FCommandLine::Get() );
 
@@ -252,15 +264,6 @@ int LightmassMain(int argc, ANSICHAR* argv[])
 			FPlatformProcess::Sleep(1.0f);
 		} while (FPlatformTime::Seconds() - StartTime < 5.0f);
 	}
-
-	FEngineLoop::AppPreExit();
-	FModuleManager::Get().UnloadModulesAtShutdown();
-
-	// Similar to CL 16324633. If we shut down FTaskGraphInterface right here, in AppExit
-	//   FThreadStats::StopThread could get called and that will try and execute some stats
-	//   commands which use FTaskGraphInterface and --> crashes. FEngineLoop::AppExit() calls
-	//   FTaskGraphInterface::Shutdown() after calling FThreadStats::StopThread() if needed.
-	FEngineLoop::AppExit();
 #endif
 
 	return 0;
