Fixed Lightmass wrong string formatting

#GitHub 9153
#9153

diff --git a/Engine/Source/Programs/UnrealLightmass/Private/ImportExport/Exporter.cpp b/Engine/Source/Programs/UnrealLightmass/Private/ImportExport/Exporter.cpp
--- a/Engine/Source/Programs/UnrealLightmass/Private/ImportExport/Exporter.cpp
+++ b/Engine/Source/Programs/UnrealLightmass/Private/ImportExport/Exporter.cpp
@@ -285,7 +285,7 @@ namespace Lightmass
 			}
 			else
 			{
-				Files[CompIndex] = IFileManager::Get().CreateFileWriter(*FString::Printf(TEXT("%s_Dir%d.bmp"), BitmapBaseName));
+				Files[CompIndex] = IFileManager::Get().CreateFileWriter(*FString::Printf(TEXT("%s_Dir%d.bmp"), BitmapBaseName, CompIndex));
 			}
 		}
 
