Lightmass LightingMesh.cpp fix

The 2D tri/AABB intersect was using FMath::{Min,Max} to compute AABBs which doesn't actually work (needs componentwise min/max) and is now deprecated; update to use appropriate min/max.

diff --git a/Engine/Source/Programs/UnrealLightmass/Private/Lighting/LightingMesh.cpp b/Engine/Source/Programs/UnrealLightmass/Private/Lighting/LightingMesh.cpp
--- a/Engine/Source/Programs/UnrealLightmass/Private/Lighting/LightingMesh.cpp
+++ b/Engine/Source/Programs/UnrealLightmass/Private/Lighting/LightingMesh.cpp
@@ -163,10 +163,10 @@ static bool AxisAlignedTriangleIntersectTriangle2d(
 	const FVector2f& V0, const FVector2f& V1, const FVector2f& V2, 
 	const FVector2f& OtherV0, const FVector2f& OtherV1, const FVector2f& OtherV2)
 {
-	const FVector2f MinFirst = FMath::Min(V0, FMath::Min(V1, V2));
-	const FVector2f MaxFirst = FMath::Max(V0, FMath::Max(V1, V2));
-	const FVector2f MinSecond = FMath::Min(OtherV0, FMath::Min(OtherV1, OtherV2));
-	const FVector2f MaxSecond = FMath::Max(OtherV0, FMath::Max(OtherV1, OtherV2));
+	const FVector2f MinFirst = FVector2f::Min(V0, FVector2f::Min(V1, V2));
+	const FVector2f MaxFirst = FVector2f::Max(V0, FVector2f::Max(V1, V2));
+	const FVector2f MinSecond = FVector2f::Min(OtherV0, FVector2f::Min(OtherV1, OtherV2));
+	const FVector2f MaxSecond = FVector2f::Max(OtherV0, FVector2f::Max(OtherV1, OtherV2));
 
 	return !(MinFirst.X > MaxSecond.X 
 		|| MinSecond.X > MaxFirst.X 
