Replace references to UE4 with UE5

diff --git a/Engine/Source/Programs/UnrealLightmass/Private/ImportExport/Exporter.cpp b/Engine/Source/Programs/UnrealLightmass/Private/ImportExport/Exporter.cpp
--- a/Engine/Source/Programs/UnrealLightmass/Private/ImportExport/Exporter.cpp
+++ b/Engine/Source/Programs/UnrealLightmass/Private/ImportExport/Exporter.cpp
@@ -95,7 +95,7 @@ namespace Lightmass
 
 	/** 
 	 * Exports information about mesh area lights back to Unreal, 
-	 * So that UE4 can create dynamic lights to approximate the mesh area light's influence on dynamic objects.
+	 * So that UE5 can create dynamic lights to approximate the mesh area light's influence on dynamic objects.
 	 */
 	void FLightmassSolverExporter::ExportMeshAreaLightData(const TIndirectArray<FMeshAreaLight>& MeshAreaLights, float MeshAreaLightGeneratedDynamicLightSurfaceOffset) const
 	{
diff --git a/Engine/Source/Programs/UnrealLightmass/Private/ImportExport/Importer.h b/Engine/Source/Programs/UnrealLightmass/Private/ImportExport/Importer.h
--- a/Engine/Source/Programs/UnrealLightmass/Private/ImportExport/Importer.h
+++ b/Engine/Source/Programs/UnrealLightmass/Private/ImportExport/Importer.h
@@ -21,7 +21,7 @@ protected:
 	 *
 	 * @param Key Key of object
 	 * @param Version Version of object to load
-	 * @param Extension Type of object to load (@todo UE4: This could be removed if Version could imply extension)
+	 * @param Extension Type of object to load (@todo UE5: This could be removed if Version could imply extension)
 	 * @return The object that was loaded or found, or NULL if the Guid failed
 	 */
 	template <class ObjType, class LookupMapType, class KeyType>
@@ -64,7 +64,7 @@ public:
 	 *
 	 * @param Guid Guid of object
 	 * @param Version Version of object to load
-	 * @param Extension Type of object to load (@todo UE4: This could be removed if Version could imply extension)
+	 * @param Extension Type of object to load (@todo UE5: This could be removed if Version could imply extension)
 	 * @return The object that was loaded or found, or NULL if the Guid failed
 	 */
 	template <class ObjType, class LookupMapType>
@@ -78,7 +78,7 @@ public:
 	 *
 	 * @param Hash Id of object
 	 * @param Version Version of object to load
-	 * @param Extension Type of object to load (@todo UE4: This could be removed if Version could imply extension)
+	 * @param Extension Type of object to load (@todo UE5: This could be removed if Version could imply extension)
 	 * @return The object that was loaded or found, or NULL if the Guid failed
 	 */
 	template <class ObjType, class LookupMapType>
@@ -179,7 +179,7 @@ bool FLightmassImporter::ImportGuidArray( ArrayType& Array, int32 Count, const L
  *
  * @param Key Key of object
  * @param Version Version of object to load
- * @param Extension Type of object to load (@todo UE4: This could be removed if Version could imply extension)
+ * @param Extension Type of object to load (@todo UE5: This could be removed if Version could imply extension)
  * @return The object that was loaded or found, or NULL if the Guid failed
  */
 template <class ObjType, class LookupMapType, class KeyType>
diff --git a/Engine/Source/Programs/UnrealLightmass/Private/ImportExport/LightmassSwarm.h b/Engine/Source/Programs/UnrealLightmass/Private/ImportExport/LightmassSwarm.h
--- a/Engine/Source/Programs/UnrealLightmass/Private/ImportExport/LightmassSwarm.h
+++ b/Engine/Source/Programs/UnrealLightmass/Private/ImportExport/LightmassSwarm.h
@@ -25,7 +25,7 @@ namespace Lightmass
 #define SWARM_ENABLE_CHANNEL_WRITES	1
 
 /** Flags to use when opening the different kinds of output channels */
-/** MUST PAIR APPROPRIATELY WITH THE SAME FLAGS IN UE4 */
+/** MUST PAIR APPROPRIATELY WITH THE SAME FLAGS IN UE5 */
 static const int32 LM_TEXTUREMAPPING_CHANNEL_FLAGS	= NSwarm::SWARM_JOB_CHANNEL_WRITE;
 static const int32 LM_VOLUMESAMPLES_CHANNEL_FLAGS		= NSwarm::SWARM_JOB_CHANNEL_WRITE;
 static const int32 LM_PRECOMPUTEDVISIBILITY_CHANNEL_FLAGS	= NSwarm::SWARM_JOB_CHANNEL_WRITE;
@@ -35,7 +35,7 @@ static const int32 LM_MESHAREALIGHT_CHANNEL_FLAGS		= NSwarm::SWARM_JOB_CHANNEL_W
 static const int32 LM_DEBUGOUTPUT_CHANNEL_FLAGS		= NSwarm::SWARM_JOB_CHANNEL_WRITE;
 
 /** Flags to use when opening the different kinds of input channels */
-/** MUST PAIR APPROPRIATELY WITH THE SAME FLAGS IN UE4 */
+/** MUST PAIR APPROPRIATELY WITH THE SAME FLAGS IN UE5 */
 #if LM_COMPRESS_INPUT_DATA
 	static const int32 LM_SCENE_CHANNEL_FLAGS			= NSwarm::SWARM_JOB_CHANNEL_READ | NSwarm::SWARM_CHANNEL_MISC_ENABLE_COMPRESSION;
 	static const int32 LM_STATICMESH_CHANNEL_FLAGS	= NSwarm::SWARM_CHANNEL_READ | NSwarm::SWARM_CHANNEL_MISC_ENABLE_COMPRESSION;
@@ -206,25 +206,25 @@ public:
 	bool	RequestTask( FGuid& TaskGuid, uint32 WaitTime = uint32(-1) );
 
 	/**
-	 * Accepts a requested task. This will also notify UE4.
+	 * Accepts a requested task. This will also notify UE5.
 	 * @param TaskGuid	The task that is being accepted
 	 */
 	void	AcceptTask( const FGuid& TaskGuid );
 
 	/**
-	 * Rejects a requested task. This will also notify UE4.
+	 * Rejects a requested task. This will also notify UE5.
 	 * @param TaskGuid	The task that is being rejected
 	 */
 	void	RejectTask( const FGuid& TaskGuid );
 
 	/**
-	 * Tells Swarm that the task is completed and all results have been fully exported. This will also notify UE4.
+	 * Tells Swarm that the task is completed and all results have been fully exported. This will also notify UE5.
 	 * @param TaskGuid	A guid that identifies the task that has been completed
 	 */
 	void	TaskCompleted( const FGuid& TaskGuid );
 
 	/**
-	 * Tells Swarm that the task has failed. This will also notify UE4.
+	 * Tells Swarm that the task has failed. This will also notify UE5.
 	 * @param TaskGuid	A guid that identifies the task that has failed
 	 */
 	void	TaskFailed( const FGuid& TaskGuid );
diff --git a/Engine/Source/Programs/UnrealLightmass/Private/ImportExport/Material.h b/Engine/Source/Programs/UnrealLightmass/Private/ImportExport/Material.h
--- a/Engine/Source/Programs/UnrealLightmass/Private/ImportExport/Material.h
+++ b/Engine/Source/Programs/UnrealLightmass/Private/ImportExport/Material.h
@@ -3,7 +3,7 @@
 #pragma once
 
 #include "CoreMinimal.h"
-// @todo UE4: This seems wrong to me to be needed (scene uses mesh, not the other way around), and is only needed for the MAX_TEXCOORDS
+// @todo UE5: This seems wrong to me to be needed (scene uses mesh, not the other way around), and is only needed for the MAX_TEXCOORDS
 #include "MaterialExport.h"
 #include "SceneExport.h"
 #include "Texture.h"
diff --git a/Engine/Source/Programs/UnrealLightmass/Private/ImportExport/Mesh.h b/Engine/Source/Programs/UnrealLightmass/Private/ImportExport/Mesh.h
--- a/Engine/Source/Programs/UnrealLightmass/Private/ImportExport/Mesh.h
+++ b/Engine/Source/Programs/UnrealLightmass/Private/ImportExport/Mesh.h
@@ -3,7 +3,7 @@
 #pragma once
 
 #include "CoreMinimal.h"
-// @todo UE4: This seems wrong to me to be needed (scene uses mesh, not the other way around), and is only needed for the MAX_TEXCOORDS
+// @todo UE5: This seems wrong to me to be needed (scene uses mesh, not the other way around), and is only needed for the MAX_TEXCOORDS
 #include "MeshExport.h"
 #include "SceneExport.h"
 
diff --git a/Engine/Source/Programs/UnrealLightmass/Private/Lighting/LightingMesh.cpp b/Engine/Source/Programs/UnrealLightmass/Private/Lighting/LightingMesh.cpp
--- a/Engine/Source/Programs/UnrealLightmass/Private/Lighting/LightingMesh.cpp
+++ b/Engine/Source/Programs/UnrealLightmass/Private/Lighting/LightingMesh.cpp
@@ -251,7 +251,7 @@ void FStaticLightingMesh::CreateMeshAreaLights(
 	// we floor the min and max because we need to see which integer wrap of UVs it falls into
 	// so, if we had range .2 to .8, the floors would both go to 0, then add 1 to account for that one
 	// if we have range -.2 to .3, we need space for the -1 .. 0 wrap, and the 0 to 1 wrap (ie 2 iterations)
-	// @todo UE4: Actually, if the Min was used to modify the UV when looping through the Corners array,
+	// @todo UE5: Actually, if the Min was used to modify the UV when looping through the Corners array,
 	// we wouldn't need full integer ranges
 	const int32 NumIterationsX = (FMath::FloorToInt(MaxUV.X) - FMath::FloorToInt(MinUV.X)) + 1;
 	const int32 NumIterationsY = (FMath::FloorToInt(MaxUV.Y) - FMath::FloorToInt(MinUV.Y)) + 1;
diff --git a/Engine/Source/Programs/UnrealLightmass/Private/Lighting/LightingSystem.cpp b/Engine/Source/Programs/UnrealLightmass/Private/Lighting/LightingSystem.cpp
--- a/Engine/Source/Programs/UnrealLightmass/Private/Lighting/LightingSystem.cpp
+++ b/Engine/Source/Programs/UnrealLightmass/Private/Lighting/LightingSystem.cpp
@@ -118,7 +118,7 @@ FStaticLightingMappingContext::~FStaticLightingMappingContext()
  * Initializes this static lighting system, and builds static lighting based on the provided options.
  * @param InOptions		- The static lighting build options.
  * @param InScene		- The scene containing all the lights and meshes
- * @param InExporter	- The exporter used to send completed data back to UE4
+ * @param InExporter	- The exporter used to send completed data back to UE5
  * @param InNumThreads	- Number of concurrent threads to use for lighting building
  */
 FStaticLightingSystem::FStaticLightingSystem(const FLightingBuildOptions& InOptions, FScene& InScene, FLightmassSolverExporter& InExporter, int32 InNumThreads)
diff --git a/Engine/Source/Programs/UnrealLightmass/Private/Lighting/SampleVolume.cpp b/Engine/Source/Programs/UnrealLightmass/Private/Lighting/SampleVolume.cpp
--- a/Engine/Source/Programs/UnrealLightmass/Private/Lighting/SampleVolume.cpp
+++ b/Engine/Source/Programs/UnrealLightmass/Private/Lighting/SampleVolume.cpp
@@ -581,7 +581,7 @@ void FStaticLightingSystem::ProcessVolumeSamplesTask(const FVolumeSamplesTaskDes
 
 /** 
  * Interpolates lighting from the volume lighting samples to a vertex. 
- * This mirrors FPrecomputedLightVolume::InterpolateIncidentRadiance in UE4, used for visualizing interpolation from the lighting volume on surfaces.
+ * This mirrors FPrecomputedLightVolume::InterpolateIncidentRadiance in UE5, used for visualizing interpolation from the lighting volume on surfaces.
  */
 FGatheredLightSample FStaticLightingSystem::InterpolatePrecomputedVolumeIncidentRadiance(const FStaticLightingVertex& Vertex, float SampleRadius, FCoherentRayCache& RayCache, bool bDebugThisTexel) const
 {
diff --git a/Engine/Source/Programs/UnrealLightmass/Private/LightmassCore/Misc/LMDebug.cpp b/Engine/Source/Programs/UnrealLightmass/Private/LightmassCore/Misc/LMDebug.cpp
--- a/Engine/Source/Programs/UnrealLightmass/Private/LightmassCore/Misc/LMDebug.cpp
+++ b/Engine/Source/Programs/UnrealLightmass/Private/LightmassCore/Misc/LMDebug.cpp
@@ -49,9 +49,9 @@ void appHandleCriticalError()
 #if !UE_BUILD_DEBUG
 		{
 			TCHAR ReportDumpVersion[] = TEXT("4");
-			TCHAR ReportDumpFilename[] = TEXT("UE4AutoReportDump.txt");
+			TCHAR ReportDumpFilename[] = TEXT("UE5AutoReportDump.txt");
 			TCHAR AutoReportExe[] = TEXT("../DotNET/AutoReporter.exe");
-			TCHAR IniDumpFilename[] = TEXT("UE4AutoReportIniDump.txt");
+			TCHAR IniDumpFilename[] = TEXT("UE5AutoReportIniDump.txt");
 
 			FArchive* AutoReportFile = IFileManager::Get().CreateFileWriter(ReportDumpFilename);
 			if (AutoReportFile != NULL)
