#9696

diff --git a/Engine/Source/Programs/UnrealLightmass/Private/CPUSolver/CPUSolver.cpp b/Engine/Source/Programs/UnrealLightmass/Private/CPUSolver/CPUSolver.cpp
--- a/Engine/Source/Programs/UnrealLightmass/Private/CPUSolver/CPUSolver.cpp
+++ b/Engine/Source/Programs/UnrealLightmass/Private/CPUSolver/CPUSolver.cpp
@@ -152,7 +152,7 @@ void ReportLogFile()
  * 
  * @param SceneGuid		Guid of the scene to process
  * @param NumThreads	Number of concurrent threads to use for lighting building
- * @param bDumpTextures	If true, 2d lightmaps will be dumped to 
+ * @param bDumpTextures	If true, 2d lightmaps will be dumped too 
  */
 void BuildStaticLighting(const FGuid& SceneGuid, int32 NumThreads, bool bDumpTextures)
 {
diff --git a/Engine/Source/Programs/UnrealLightmass/Private/ImportExport/LightmassScene.cpp b/Engine/Source/Programs/UnrealLightmass/Private/ImportExport/LightmassScene.cpp
--- a/Engine/Source/Programs/UnrealLightmass/Private/ImportExport/LightmassScene.cpp
+++ b/Engine/Source/Programs/UnrealLightmass/Private/ImportExport/LightmassScene.cpp
@@ -9,7 +9,7 @@
 namespace Lightmass
 {
 
-/** Copy ctor that doesn't modify padding in FSceneFileHeader. */
+/** Copy constructor that doesn't modify padding in FSceneFileHeader. */
 FSceneFileHeader::FSceneFileHeader(const FSceneFileHeader& Other)
 {
 	/** FourCC cookie: 'SCEN' */
