CPULM: Put the hash of UnrealLightmass executable into Swarm channel names to avoid getting stale Swarm cache content from a different build

diff --git a/Engine/Source/Editor/UnrealEd/Private/Lightmass/Lightmass.cpp b/Engine/Source/Editor/UnrealEd/Private/Lightmass/Lightmass.cpp
--- a/Engine/Source/Editor/UnrealEd/Private/Lightmass/Lightmass.cpp
+++ b/Engine/Source/Editor/UnrealEd/Private/Lightmass/Lightmass.cpp
@@ -53,6 +53,7 @@
 #include "Misc/MessageDialog.h"
 #include "Modules/ModuleManager.h"
 #include "ImageCoreUtils.h"
+#include "Misc/FileHelper.h"
 
 extern FSwarmDebugOptions GSwarmDebugOptions;
 
@@ -505,6 +506,13 @@ FLightmassExporter::FLightmassExporter( UWorld* InWorld )
 	{
 		SceneGuid = FGuid::NewGuid();
 	}
+
+	// Compute the hash of UnrealLightmass executable and store it into LightmassExecutableHash
+	// Guaranteed success as the executable has been checked by CheckLightmassExecutableVersion()
+	TArray<uint8> Bytes;
+	FFileHelper::LoadFileToArray(Bytes, *FPlatformProcess::GenerateApplicationPath(TEXT("UnrealLightmass"), EBuildConfiguration::Development));
+	FSHA1::HashBuffer(&Bytes[0], Bytes.Num(), LightmassExecutableHash.Hash);
+	
 	ChannelName = Lightmass::CreateChannelName(SceneGuid, Lightmass::LM_SCENE_VERSION, Lightmass::LM_SCENE_EXTENSION);
 }
 
@@ -1226,7 +1234,7 @@ void FLightmassExporter::WriteStaticMeshes()
 		BaseMeshData.Guid = StaticMesh->GetLightingGuid();
 
 		// create a channel name to write the mesh out to
-		FString NewChannelName = Lightmass::CreateChannelName(BaseMeshData.Guid, Lightmass::LM_STATICMESH_VERSION, Lightmass::LM_STATICMESH_EXTENSION);
+		FString NewChannelName = Lightmass::CreateChannelNameWithLMExecutableHash(BaseMeshData.Guid, Lightmass::LM_STATICMESH_VERSION, LightmassExecutableHash, Lightmass::LM_STATICMESH_EXTENSION);
 
 		// Warn the user if there is an invalid lightmap UV channel specified.
 		if( StaticMesh->GetLightMapCoordinateIndex() > 0
@@ -1367,7 +1375,7 @@ void FLightmassExporter::BuildMaterialMap(UMaterialInterface* Material)
 		GetMaterialHash(Material, MaterialHash);
 
 		// create a channel name to write the material out to
-		FString NewChannelName = Lightmass::CreateChannelName(MaterialHash, Lightmass::LM_MATERIAL_VERSION, Lightmass::LM_MATERIAL_EXTENSION);
+		FString NewChannelName = Lightmass::CreateChannelNameWithLMExecutableHash(MaterialHash, Lightmass::LM_MATERIAL_VERSION, LightmassExecutableHash, Lightmass::LM_MATERIAL_EXTENSION);
 
 		// only export the material if it's not currently in the cache
 		int32 ErrorCode;
diff --git a/Engine/Source/Editor/UnrealEd/Private/Lightmass/Lightmass.h b/Engine/Source/Editor/UnrealEd/Private/Lightmass/Lightmass.h
--- a/Engine/Source/Editor/UnrealEd/Private/Lightmass/Lightmass.h
+++ b/Engine/Source/Editor/UnrealEd/Private/Lightmass/Lightmass.h
@@ -199,6 +199,7 @@ private:
 	NSwarm::FSwarmInterface&	Swarm;
 	bool						bSwarmConnectionIsValid;
 	FGuid						SceneGuid;
+	FSHAHash					LightmassExecutableHash;
 	FString						ChannelName;
 
 	TArray<FBox>				ImportanceVolumes;
diff --git a/Engine/Source/Programs/UnrealLightmass/Private/ImportExport/Importer.cpp b/Engine/Source/Programs/UnrealLightmass/Private/ImportExport/Importer.cpp
--- a/Engine/Source/Programs/UnrealLightmass/Private/ImportExport/Importer.cpp
+++ b/Engine/Source/Programs/UnrealLightmass/Private/ImportExport/Importer.cpp
@@ -3,7 +3,8 @@
 #include "Importer.h"
 #include "LightmassSwarm.h"
 #include "LightmassScene.h"
-
+#include "Misc/FileHelper.h"
+#include "Misc/SecureHash.h"
 
 namespace Lightmass
 {
@@ -12,6 +13,10 @@ FLightmassImporter::FLightmassImporter( FLightmassSwarm* InSwarm )
 :	Swarm( InSwarm )
 ,	LevelScale(0.0f)
 {
+	// Compute the hash of UnrealLightmass executable (which is the current running process) and store it into LightmassExecutableHash
+	TArray<uint8> Bytes;
+	FFileHelper::LoadFileToArray(Bytes, FPlatformProcess::ExecutablePath());
+	FSHA1::HashBuffer(&Bytes[0], Bytes.Num(), LightmassExecutableHash.Hash);
 }
 
 FLightmassImporter::~FLightmassImporter()
diff --git a/Engine/Source/Programs/UnrealLightmass/Private/ImportExport/Importer.h b/Engine/Source/Programs/UnrealLightmass/Private/ImportExport/Importer.h
--- a/Engine/Source/Programs/UnrealLightmass/Private/ImportExport/Importer.h
+++ b/Engine/Source/Programs/UnrealLightmass/Private/ImportExport/Importer.h
@@ -21,11 +21,12 @@ protected:
 	 *
 	 * @param Key Key of object
 	 * @param Version Version of object to load
+	 * @param LMExecutableHash Hash of Lightmass executable to avoid stale cache content
 	 * @param Extension Type of object to load (@todo UE5: This could be removed if Version could imply extension)
 	 * @return The object that was loaded or found, or NULL if the Guid failed
 	 */
 	template <class ObjType, class LookupMapType, class KeyType>
-	ObjType* ConditionalImportObjectWithKey(const KeyType& Key, const int32 Version, const TCHAR* Extension, int32 ChannelFlags, LookupMapType& LookupMap);
+	ObjType* ConditionalImportObjectWithKey(const KeyType& Key, const int32 Version, const FSHAHash& LMExecutableHash, const TCHAR* Extension, int32 ChannelFlags, LookupMapType& LookupMap);
 
 public:
 
@@ -70,7 +71,7 @@ public:
 	template <class ObjType, class LookupMapType>
 	ObjType* ConditionalImportObject(const FGuid& Guid, const int32 Version, const TCHAR* Extension, int32 ChannelFlags, LookupMapType& LookupMap)
 	{
-		return ConditionalImportObjectWithKey<ObjType, LookupMapType, FGuid>(Guid, Version, Extension, ChannelFlags, LookupMap);
+		return ConditionalImportObjectWithKey<ObjType, LookupMapType, FGuid>(Guid, Version, LightmassExecutableHash, Extension, ChannelFlags, LookupMap);
 	}
 
 	/**
@@ -84,7 +85,7 @@ public:
 	template <class ObjType, class LookupMapType>
 	ObjType* ConditionalImportObject(const FSHAHash& Hash, const int32 Version, const TCHAR* Extension, int32 ChannelFlags, LookupMapType& LookupMap)
 	{
-		return ConditionalImportObjectWithKey<ObjType, LookupMapType, FSHAHash>(Hash, Version, Extension, ChannelFlags, LookupMap);
+		return ConditionalImportObjectWithKey<ObjType, LookupMapType, FSHAHash>(Hash, Version, LightmassExecutableHash, Extension, ChannelFlags, LookupMap);
 	}
 
 	void SetLevelScale(float InScale) { LevelScale = InScale; }
@@ -109,6 +110,7 @@ public:
 private:
 
 	class FLightmassSwarm*	Swarm;
+	FSHAHash LightmassExecutableHash;
 
 	TMap<FGuid,class FLight*>										Lights;
 	TMap<FGuid,class FStaticMesh*>									StaticMeshes;
@@ -183,14 +185,14 @@ bool FLightmassImporter::ImportGuidArray( ArrayType& Array, int32 Count, const L
  * @return The object that was loaded or found, or NULL if the Guid failed
  */
 template <class ObjType, class LookupMapType, class KeyType>
-ObjType* FLightmassImporter::ConditionalImportObjectWithKey(const KeyType& Key, const int32 Version, const TCHAR* Extension, int32 ChannelFlags, LookupMapType& LookupMap)
+ObjType* FLightmassImporter::ConditionalImportObjectWithKey(const KeyType& Key, const int32 Version, const FSHAHash& LMExecutableHash, const TCHAR* Extension, int32 ChannelFlags, LookupMapType& LookupMap)
 {
 	// look to see if it exists already
 	ObjType* Obj = LookupMap.FindRef(Key);
 	if (Obj == NULL)
 	{
 		// open a new channel and make it current
-		if (Swarm->OpenChannel(*CreateChannelName(Key, Version, Extension), ChannelFlags, true) >= 0)
+		if (Swarm->OpenChannel(*CreateChannelNameWithLMExecutableHash(Key, Version, LMExecutableHash, Extension), ChannelFlags, true) >= 0)
 		{
 			Obj = new ObjType;
 
diff --git a/Engine/Source/Programs/UnrealLightmass/Public/ImportExport.h b/Engine/Source/Programs/UnrealLightmass/Public/ImportExport.h
--- a/Engine/Source/Programs/UnrealLightmass/Public/ImportExport.h
+++ b/Engine/Source/Programs/UnrealLightmass/Public/ImportExport.h
@@ -501,26 +501,32 @@ namespace Lightmass
 	 *
 	 * @param Guid Unique ID of the channel
 	 * @param Version Version of the data inside
+	 * @param LMExecutableHash Hash of UnrealLightmass executable, to avoid stale cache content from a different build
 	 * @param Extension Type of the data
 	 *
 	 * @return Standard channel name
-	 */
+	*/
+	FORCEINLINE FString CreateChannelNameWithLMExecutableHash(const FGuid& Guid, const int32 Version, const FSHAHash& LMExecutableHash, const FString& Extension)
+	{
+		return FString::Printf(TEXT("v%d.%08X%08X%08X%08X.%s.%s"), Version, Guid.A, Guid.B, Guid.C, Guid.D, *LMExecutableHash.ToString(), *Extension);
+	}
 	FORCEINLINE FString CreateChannelName(const FGuid& Guid, const int32 Version, const FString& Extension)
 	{
 		return FString::Printf(TEXT("v%d.%08X%08X%08X%08X.%s"), Version, Guid.A, Guid.B, Guid.C, Guid.D, *Extension);
 	}
 
 	/**
-	 * Creates a standardized channel name based a FSA Hash, version and type
+	 * Creates a standardized channel name based a FSA Hash, version and type. Used for exporting materials
 	 *
 	 * @param Hash Unique ID of the channel
 	 * @param Version Version of the data inside
+	 * @param LMExecutableHash Hash of UnrealLightmass executable, to avoid stale cache content from a different build
 	 * @param Extension Type of the data
 	 *
 	 * @return Standard channel name
 	 */
-	FORCEINLINE FString CreateChannelName(const FSHAHash& Hash, const int32 Version, const FString& Extension)
+	FORCEINLINE FString CreateChannelNameWithLMExecutableHash(const FSHAHash& Hash, const int32 Version, const FSHAHash& LMExecutableHash, const FString& Extension)
 	{
-		return FString::Printf(TEXT("v%d.%s.%s"), Version, *Hash.ToString(), *Extension);
+		return FString::Printf(TEXT("v%d.%s.%s.%s"), Version, *Hash.ToString(), *LMExecutableHash.ToString(), *Extension);
 	}
 } // namespace Lightmass
