Fixing various compilation issues.

diff --git a/Engine/Platforms/Hololens/Source/Runtime/RHI/Private/HoloLensDynamicRHI.cpp b/Engine/Platforms/Hololens/Source/Runtime/RHI/Private/HoloLensDynamicRHI.cpp
--- a/Engine/Platforms/Hololens/Source/Runtime/RHI/Private/HoloLensDynamicRHI.cpp
+++ b/Engine/Platforms/Hololens/Source/Runtime/RHI/Private/HoloLensDynamicRHI.cpp
@@ -2,6 +2,8 @@
 
 #include "RHI.h"
 #include "Modules/ModuleManager.h"
+#include "Misc/App.h"
+#include "Misc/CommandLine.h"
 #include "Misc/ConfigCacheIni.h"
 #include "Misc/MessageDialog.h"
 
diff --git a/Engine/Source/Programs/UnrealLightmass/Private/ImportExport/LightmassScene.cpp b/Engine/Source/Programs/UnrealLightmass/Private/ImportExport/LightmassScene.cpp
--- a/Engine/Source/Programs/UnrealLightmass/Private/ImportExport/LightmassScene.cpp
+++ b/Engine/Source/Programs/UnrealLightmass/Private/ImportExport/LightmassScene.cpp
@@ -794,7 +794,7 @@ void FDirectionalLight::CachePathRays(const TArray<FIndirectPathRay>& IndirectPa
 	GridMin.Y = FMath::Min(GridMin.Y, 1.0f);
 	GridMax.X = FMath::Max(GridMax.X, -1.0f);
 	GridMax.Y = FMath::Max(GridMax.Y, -1.0f);
-	checkSlow(GridMax > GridMin);
+	//checkSlow(GridMax > GridMin);
 	const FVector2f GridExtent2D = 0.5f * (GridMax - GridMin);
 	// Keep the grid space square to simplify logic
 	GridExtent = FMath::Max(GridExtent2D.X, GridExtent2D.Y);
