Removing bad Launch include paths from programs.

diff --git a/Engine/Plugins/Runtime/AndroidPermission/Source/AndroidPermission/AndroidPermission.Build.cs b/Engine/Plugins/Runtime/AndroidPermission/Source/AndroidPermission/AndroidPermission.Build.cs
--- a/Engine/Plugins/Runtime/AndroidPermission/Source/AndroidPermission/AndroidPermission.Build.cs
+++ b/Engine/Plugins/Runtime/AndroidPermission/Source/AndroidPermission/AndroidPermission.Build.cs
@@ -7,14 +7,8 @@ public class AndroidPermission : ModuleRules
 {
 	public AndroidPermission(ReadOnlyTargetRules Target) : base(Target)
 	{
-		
-		PublicIncludePaths.AddRange(
-			new string[] {
-				"Runtime/Launch/Public"
-			}
-			);
-			
-		
+		PublicIncludePathModuleNames.Add("Launch");
+
 		PublicDependencyModuleNames.AddRange(
 			new string[]
 			{
diff --git a/Engine/Plugins/Runtime/GooglePAD/Source/GooglePAD/GooglePAD.Build.cs b/Engine/Plugins/Runtime/GooglePAD/Source/GooglePAD/GooglePAD.Build.cs
--- a/Engine/Plugins/Runtime/GooglePAD/Source/GooglePAD/GooglePAD.Build.cs
+++ b/Engine/Plugins/Runtime/GooglePAD/Source/GooglePAD/GooglePAD.Build.cs
@@ -8,12 +8,8 @@ public class GooglePAD : ModuleRules
 {
 	public GooglePAD(ReadOnlyTargetRules Target) : base(Target)
 	{
-		PublicIncludePaths.AddRange(
-			new string[] {
-				"Runtime/Launch/Public"
-			}
-			);			
-		
+		PublicIncludePathModuleNames.Add("Launch");
+
 		PublicDependencyModuleNames.AddRange(
 			new string[]
 			{
diff --git a/Engine/Source/Developer/Datasmith/DatasmithExporter/DatasmithExporter.Build.cs b/Engine/Source/Developer/Datasmith/DatasmithExporter/DatasmithExporter.Build.cs
--- a/Engine/Source/Developer/Datasmith/DatasmithExporter/DatasmithExporter.Build.cs
+++ b/Engine/Source/Developer/Datasmith/DatasmithExporter/DatasmithExporter.Build.cs
@@ -35,8 +35,7 @@ namespace UnrealBuildTool.Rules
 				}
 			);
 
-			PrivateIncludePaths.Add("Runtime/Launch/Public");
-			PrivateIncludePaths.Add("Runtime/Launch/Private");      // For LaunchEngineLoop.cpp include
+			PrivateIncludePathModuleNames.Add("Launch");
 
 			// PrecompileForTargets = PrecompileTargetsType.Any;
 		}
diff --git a/Engine/Source/Programs/BenchmarkTool/BenchmarkTool.Build.cs b/Engine/Source/Programs/BenchmarkTool/BenchmarkTool.Build.cs
--- a/Engine/Source/Programs/BenchmarkTool/BenchmarkTool.Build.cs
+++ b/Engine/Source/Programs/BenchmarkTool/BenchmarkTool.Build.cs
@@ -6,10 +6,7 @@ public class BenchmarkTool : ModuleRules
 {
 	public BenchmarkTool(ReadOnlyTargetRules Target) : base(Target)
 	{
-		PublicIncludePaths.Add("Runtime/Launch/Public");
-
-		PrivateIncludePaths.Add("Runtime/Launch/Private");		// For LaunchEngineLoop.cpp include
-
+		PublicIncludePathModuleNames.Add("Launch");
 		PrivateDependencyModuleNames.Add("Core");
 		PrivateDependencyModuleNames.Add("Projects");
 	}
diff --git a/Engine/Source/Programs/BlankProgram/BlankProgram.Build.cs b/Engine/Source/Programs/BlankProgram/BlankProgram.Build.cs
--- a/Engine/Source/Programs/BlankProgram/BlankProgram.Build.cs
+++ b/Engine/Source/Programs/BlankProgram/BlankProgram.Build.cs
@@ -6,10 +6,7 @@ public class BlankProgram : ModuleRules
 {
 	public BlankProgram(ReadOnlyTargetRules Target) : base(Target)
 	{
-		PublicIncludePaths.Add("Runtime/Launch/Public");
-
-		PrivateIncludePaths.Add("Runtime/Launch/Private");		// For LaunchEngineLoop.cpp include
-
+		PublicIncludePathModuleNames.Add("Launch");
 		PrivateDependencyModuleNames.Add("Core");
 		PrivateDependencyModuleNames.Add("Projects");
 	}
diff --git a/Engine/Source/Programs/BuildPatchTool/BuildPatchTool.Build.cs b/Engine/Source/Programs/BuildPatchTool/BuildPatchTool.Build.cs
--- a/Engine/Source/Programs/BuildPatchTool/BuildPatchTool.Build.cs
+++ b/Engine/Source/Programs/BuildPatchTool/BuildPatchTool.Build.cs
@@ -6,10 +6,7 @@ public class BuildPatchTool : ModuleRules
 {
 	public BuildPatchTool( ReadOnlyTargetRules Target ) : base(Target)
 	{
-		PublicIncludePaths.Add("Runtime/Launch/Public");
-
-		// For LaunchEngineLoop.cpp include
-		PrivateIncludePaths.Add("Runtime/Launch/Private");
+		PublicIncludePathModuleNames.Add("Launch");
 
 		PrivateDependencyModuleNames.AddRange(
 			new string[]
diff --git a/Engine/Source/Programs/ChaosVisualDebugger/ChaosVisualDebugger.Build.cs b/Engine/Source/Programs/ChaosVisualDebugger/ChaosVisualDebugger.Build.cs
--- a/Engine/Source/Programs/ChaosVisualDebugger/ChaosVisualDebugger.Build.cs
+++ b/Engine/Source/Programs/ChaosVisualDebugger/ChaosVisualDebugger.Build.cs
@@ -6,10 +6,7 @@ public class ChaosVisualDebugger : ModuleRules
 {
 	public ChaosVisualDebugger(ReadOnlyTargetRules Target) : base(Target)
 	{
-		PublicIncludePaths.Add("Runtime/Launch/Public");
-
-		PrivateIncludePaths.Add("Runtime/Launch/Private");		// For LaunchEngineLoop.cpp include
-
+		PublicIncludePathModuleNames.Add("Launch");
 		PrivateDependencyModuleNames.Add("Core");
 		PrivateDependencyModuleNames.Add("Projects");
 
diff --git a/Engine/Source/Programs/CrashReportClient/CrashReportClient.Build.cs b/Engine/Source/Programs/CrashReportClient/CrashReportClient.Build.cs
--- a/Engine/Source/Programs/CrashReportClient/CrashReportClient.Build.cs
+++ b/Engine/Source/Programs/CrashReportClient/CrashReportClient.Build.cs
@@ -6,14 +6,7 @@ public class CrashReportClient : ModuleRules
 {
 	public CrashReportClient(ReadOnlyTargetRules Target) : base(Target)
 	{
-		PublicIncludePaths.AddRange
-		(
-			new string[] 
-			{ 
-				"Runtime/Launch/Public",
-				"Programs/CrashReportClient/Private",
-			}
-		);
+		PublicIncludePathModuleNames.Add("Launch");
 
 		PrivateDependencyModuleNames.AddRange(
 			new string[] 
@@ -54,8 +47,6 @@ public class CrashReportClient : ModuleRules
 			);
 		}
 
-		PrivateIncludePaths.Add("Runtime/Launch/Private");		// For LaunchEngineLoop.cpp include
-
 		PrivateDefinitions.AddRange(
 			new string[]
 			{
diff --git a/Engine/Source/Programs/DerivedDataBuildWorker/DerivedDataBuildWorker.Build.cs b/Engine/Source/Programs/DerivedDataBuildWorker/DerivedDataBuildWorker.Build.cs
--- a/Engine/Source/Programs/DerivedDataBuildWorker/DerivedDataBuildWorker.Build.cs
+++ b/Engine/Source/Programs/DerivedDataBuildWorker/DerivedDataBuildWorker.Build.cs
@@ -53,9 +53,8 @@ public class DerivedDataBuildWorker : ModuleRules
 {
 	public DerivedDataBuildWorker(ReadOnlyTargetRules Target) : base(Target)
 	{
-		PublicIncludePaths.Add("Runtime/Launch/Public");
-		PublicIncludePaths.Add("Developer/DerivedDataCache/Public");
-		PrivateIncludePaths.Add("Runtime/Launch/Private");		// For LaunchEngineLoop.cpp include
+		PublicIncludePathModuleNames.Add("Launch");
+		PublicIncludePathModuleNames.Add("DerivedDataCache");
 
 		PrivateDependencyModuleNames.AddRange(new string[]
 		{
diff --git a/Engine/Source/Programs/DerivedDataTool/DerivedDataTool.Build.cs b/Engine/Source/Programs/DerivedDataTool/DerivedDataTool.Build.cs
--- a/Engine/Source/Programs/DerivedDataTool/DerivedDataTool.Build.cs
+++ b/Engine/Source/Programs/DerivedDataTool/DerivedDataTool.Build.cs
@@ -6,9 +6,7 @@ public class DerivedDataTool : ModuleRules
 {
 	public DerivedDataTool(ReadOnlyTargetRules Target) : base(Target)
 	{
-		PublicIncludePaths.Add("Runtime/Launch/Public");
-		PrivateIncludePaths.Add("Runtime/Launch/Private"); // LaunchEngineLoop.cpp
-
+		PublicIncludePathModuleNames.Add("Launch");
 		PrivateDependencyModuleNames.Add("Core");
 		PrivateDependencyModuleNames.Add("Projects");
 		PrivateDependencyModuleNames.Add("DerivedDataCache");
diff --git a/Engine/Source/Programs/Enterprise/Datasmith/DatasmithARCHICADExporter/DatsmithUE4ArchiCAD/DatasmithUE4ArchiCAD.Build.cs b/Engine/Source/Programs/Enterprise/Datasmith/DatasmithARCHICADExporter/DatsmithUE4ArchiCAD/DatasmithUE4ArchiCAD.Build.cs
--- a/Engine/Source/Programs/Enterprise/Datasmith/DatasmithARCHICADExporter/DatsmithUE4ArchiCAD/DatasmithUE4ArchiCAD.Build.cs
+++ b/Engine/Source/Programs/Enterprise/Datasmith/DatasmithARCHICADExporter/DatsmithUE4ArchiCAD/DatasmithUE4ArchiCAD.Build.cs
@@ -11,7 +11,7 @@ namespace UnrealBuildTool.Rules
 			bUseRTTI = true;
 			PCHUsage = PCHUsageMode.NoPCHs;
 
-			PublicIncludePaths.Add("Runtime/Launch/Public");
+			PublicIncludePathModuleNames.Add("Launch");
 
 			PrivateDependencyModuleNames.AddRange(
 				new string[]
diff --git a/Engine/Source/Programs/Enterprise/Datasmith/DatasmithARCHICADExporter/UBTModules/DatasmithArchiCAD.Build.cs b/Engine/Source/Programs/Enterprise/Datasmith/DatasmithARCHICADExporter/UBTModules/DatasmithArchiCAD.Build.cs
--- a/Engine/Source/Programs/Enterprise/Datasmith/DatasmithARCHICADExporter/UBTModules/DatasmithArchiCAD.Build.cs
+++ b/Engine/Source/Programs/Enterprise/Datasmith/DatasmithARCHICADExporter/UBTModules/DatasmithArchiCAD.Build.cs
@@ -10,12 +10,8 @@ namespace UnrealBuildTool.Rules
 		public DatasmithArchiCAD(ReadOnlyTargetRules Target)
 			: base(Target)
 		{
-			PublicIncludePaths.Add("Runtime/Launch/Public");
-
-			PrivateIncludePaths.Add("Runtime/Launch/Private");
-
+			PublicIncludePathModuleNames.Add("Launch");
 			PrivateDependencyModuleNames.Add("Core");
-			//PrivateDependencyModuleNames.Add("Projects");
 		}
 	}
 }
diff --git a/Engine/Source/Programs/Enterprise/Datasmith/DatasmithCADWorker/DatasmithCADWorker.Build.cs b/Engine/Source/Programs/Enterprise/Datasmith/DatasmithCADWorker/DatasmithCADWorker.Build.cs
--- a/Engine/Source/Programs/Enterprise/Datasmith/DatasmithCADWorker/DatasmithCADWorker.Build.cs
+++ b/Engine/Source/Programs/Enterprise/Datasmith/DatasmithCADWorker/DatasmithCADWorker.Build.cs
@@ -7,9 +7,7 @@ public class DatasmithCADWorker : ModuleRules
 {
 	public DatasmithCADWorker(ReadOnlyTargetRules Target) : base(Target)
 	{
-
-		PublicIncludePaths.Add("Runtime/Launch/Public");
-		PrivateIncludePaths.Add("Runtime/Launch/Private");
+		PublicIncludePathModuleNames.Add("Launch");
 
 		PrivateDependencyModuleNames.AddRange(
 			new string[]
diff --git a/Engine/Source/Programs/Enterprise/Datasmith/DatasmithMaxExporter/DatasmithMax2017.Build.cs b/Engine/Source/Programs/Enterprise/Datasmith/DatasmithMaxExporter/DatasmithMax2017.Build.cs
--- a/Engine/Source/Programs/Enterprise/Datasmith/DatasmithMaxExporter/DatasmithMax2017.Build.cs
+++ b/Engine/Source/Programs/Enterprise/Datasmith/DatasmithMaxExporter/DatasmithMax2017.Build.cs
@@ -32,8 +32,7 @@ namespace UnrealBuildTool.Rules
 				}
 			);
 
-
-			PrivateIncludePaths.AddRange( new string[] { "Runtime/Launch/Public", "Runtime/Launch/Private", ModuleDirectory } );
+			PrivateIncludePathModuleNames.Add("Launch");
 
 			// Max SDK setup
 			{
diff --git a/Engine/Source/Programs/Enterprise/Datasmith/DatasmithNavisworksExporter/UBTModules/DatasmithNavisworks2020.Build.cs b/Engine/Source/Programs/Enterprise/Datasmith/DatasmithNavisworksExporter/UBTModules/DatasmithNavisworks2020.Build.cs
--- a/Engine/Source/Programs/Enterprise/Datasmith/DatasmithNavisworksExporter/UBTModules/DatasmithNavisworks2020.Build.cs
+++ b/Engine/Source/Programs/Enterprise/Datasmith/DatasmithNavisworksExporter/UBTModules/DatasmithNavisworks2020.Build.cs
@@ -18,14 +18,7 @@ namespace UnrealBuildTool.Rules
 				}
 			);
 
-			PrivateIncludePaths.AddRange(
-				new string[]
-				{
-					"Runtime/Launch/Public",
-					"Runtime/Launch/Private",
-					ModuleDirectory
-				}
-			);
+			PrivateIncludePathModuleNames.Add("Launch");
 		}
 
 		public abstract string GetNavisworksVersion();
diff --git a/Engine/Source/Programs/Enterprise/Datasmith/DatasmithRevitExporter/DatasmithRevit2018.Build.cs b/Engine/Source/Programs/Enterprise/Datasmith/DatasmithRevitExporter/DatasmithRevit2018.Build.cs
--- a/Engine/Source/Programs/Enterprise/Datasmith/DatasmithRevitExporter/DatasmithRevit2018.Build.cs
+++ b/Engine/Source/Programs/Enterprise/Datasmith/DatasmithRevitExporter/DatasmithRevit2018.Build.cs
@@ -20,14 +20,7 @@ namespace UnrealBuildTool.Rules
 				}
 			);
 
-			PrivateIncludePaths.AddRange(
-				new string[]
-				{
-					"Runtime/Launch/Public",
-					"Runtime/Launch/Private",
-					ModuleDirectory
-				}
-			);
+			PrivateIncludePathModuleNames.Add("Launch");
 		}
 
 		public abstract string GetRevitVersion();
diff --git a/Engine/Source/Programs/Enterprise/Datasmith/DatasmithSDK/DatasmithSDK.Build.cs b/Engine/Source/Programs/Enterprise/Datasmith/DatasmithSDK/DatasmithSDK.Build.cs
--- a/Engine/Source/Programs/Enterprise/Datasmith/DatasmithSDK/DatasmithSDK.Build.cs
+++ b/Engine/Source/Programs/Enterprise/Datasmith/DatasmithSDK/DatasmithSDK.Build.cs
@@ -8,7 +8,7 @@ namespace UnrealBuildTool.Rules
 		public DatasmithSDK(ReadOnlyTargetRules Target)
 			: base(Target)
 		{
-			PublicIncludePaths.Add("Runtime/Launch/Public");
+			PublicIncludePathModuleNames.Add("Launch");
 
 			PrivateDependencyModuleNames.AddRange(
 				new string[]
diff --git a/Engine/Source/Programs/Enterprise/Datasmith/DatasmithSolidworksExporter/DatasmithSolidworks2021.Build.cs b/Engine/Source/Programs/Enterprise/Datasmith/DatasmithSolidworksExporter/DatasmithSolidworks2021.Build.cs
--- a/Engine/Source/Programs/Enterprise/Datasmith/DatasmithSolidworksExporter/DatasmithSolidworks2021.Build.cs
+++ b/Engine/Source/Programs/Enterprise/Datasmith/DatasmithSolidworksExporter/DatasmithSolidworks2021.Build.cs
@@ -19,13 +19,7 @@ namespace UnrealBuildTool.Rules
 				}
 			);
 
-			PrivateIncludePaths.AddRange(
-				new string[]
-				{
-					"Runtime/Launch/Public",
-					"Runtime/Launch/Private",
-				}
-			);
+			PrivateIncludePathModuleNames.Add("Launch");
 		}
 
 		public abstract string GetVersion();
diff --git a/Engine/Source/Programs/EpicWebHelper/EpicWebHelper.Build.cs b/Engine/Source/Programs/EpicWebHelper/EpicWebHelper.Build.cs
--- a/Engine/Source/Programs/EpicWebHelper/EpicWebHelper.Build.cs
+++ b/Engine/Source/Programs/EpicWebHelper/EpicWebHelper.Build.cs
@@ -7,15 +7,7 @@ public class EpicWebHelper : ModuleRules
 {
 	public EpicWebHelper(ReadOnlyTargetRules Target) : base(Target)
 	{
-		PublicIncludePaths.Add("Runtime/Launch/Public");
-
-		PrivateIncludePaths.AddRange(
-			new string[]
-			{
-				"Programs/EpicWebHelper/Private",
-				"Runtime/Launch/Private",					// for LaunchEngineLoop.cpp include
-			}
-		);
+		PublicIncludePathModuleNames.Add("Launch");
 
 		PrivateDependencyModuleNames.AddRange(
 			new string[]
diff --git a/Engine/Source/Programs/HeadlessChaos/HeadlessChaos.Build.cs b/Engine/Source/Programs/HeadlessChaos/HeadlessChaos.Build.cs
--- a/Engine/Source/Programs/HeadlessChaos/HeadlessChaos.Build.cs
+++ b/Engine/Source/Programs/HeadlessChaos/HeadlessChaos.Build.cs
@@ -7,10 +7,7 @@ public class HeadlessChaos : TestModuleRules
 {
 	public HeadlessChaos(ReadOnlyTargetRules Target) : base(Target, false)
 	{
-		PublicIncludePaths.Add("Runtime/Launch/Public");
-
-		// For LaunchEngineLoop.cpp include
-		PrivateIncludePaths.Add("Runtime/Launch/Private");
+		PublicIncludePathModuleNames.Add("Launch");
 
 		// For testing access to private Chaos classes
 		PrivateIncludePaths.Add("Runtime/Experimental/Chaos/Private");
diff --git a/Engine/Source/Programs/HeadlessChaosPerf/HeadlessChaosPerf.Build.cs b/Engine/Source/Programs/HeadlessChaosPerf/HeadlessChaosPerf.Build.cs
--- a/Engine/Source/Programs/HeadlessChaosPerf/HeadlessChaosPerf.Build.cs
+++ b/Engine/Source/Programs/HeadlessChaosPerf/HeadlessChaosPerf.Build.cs
@@ -9,9 +9,7 @@ public class HeadlessChaosPerf : ModuleRules
 {
 	public HeadlessChaosPerf(ReadOnlyTargetRules Target) : base(Target)
 	{
-		PublicIncludePaths.Add("Runtime/Launch/Public");
-
-		PrivateIncludePaths.Add("Runtime/Launch/Private");
+		PublicIncludePathModuleNames.Add("Launch");
 
 		SetupModulePhysicsSupport(Target);
 
diff --git a/Engine/Source/Programs/IOS/UnrealLaunchDaemon/UnrealLaunchDaemon.Build.cs b/Engine/Source/Programs/IOS/UnrealLaunchDaemon/UnrealLaunchDaemon.Build.cs
--- a/Engine/Source/Programs/IOS/UnrealLaunchDaemon/UnrealLaunchDaemon.Build.cs
+++ b/Engine/Source/Programs/IOS/UnrealLaunchDaemon/UnrealLaunchDaemon.Build.cs
@@ -6,11 +6,7 @@ public class UnrealLaunchDaemon : ModuleRules
 {
 	public UnrealLaunchDaemon( ReadOnlyTargetRules Target ) : base(Target)
 	{
-		PublicIncludePaths.AddRange(
-			new string[] {
-				"Runtime/Launch/Public"
-			}
-		);
+		PublicIncludePathModuleNames.Add("Launch");
 
 		PrivateDependencyModuleNames.AddRange(
 			new string[] {
@@ -34,7 +30,5 @@ public class UnrealLaunchDaemon : ModuleRules
 				"Messaging",
 			}
 		);
-
-		PrivateIncludePaths.Add("Runtime/Launch/Private");
 	}
 }
diff --git a/Engine/Source/Programs/InterchangeWorker/InterchangeWorker.Build.cs b/Engine/Source/Programs/InterchangeWorker/InterchangeWorker.Build.cs
--- a/Engine/Source/Programs/InterchangeWorker/InterchangeWorker.Build.cs
+++ b/Engine/Source/Programs/InterchangeWorker/InterchangeWorker.Build.cs
@@ -14,8 +14,7 @@ public class InterchangeWorker : ModuleRules
 			throw new BuildException("InterchangeWorker program do not support target platform {0}.", Target.Platform.ToString());
 		}
 
-		PublicIncludePaths.Add("Runtime/Launch/Public");
-		PrivateIncludePaths.Add("Runtime/Launch/Private");
+		PublicIncludePathModuleNames.Add("Launch");
 
 		PrivateDependencyModuleNames.AddRange(
 			new string[]
diff --git a/Engine/Source/Programs/LiveCodingConsole/LiveCodingConsole.Build.cs b/Engine/Source/Programs/LiveCodingConsole/LiveCodingConsole.Build.cs
--- a/Engine/Source/Programs/LiveCodingConsole/LiveCodingConsole.Build.cs
+++ b/Engine/Source/Programs/LiveCodingConsole/LiveCodingConsole.Build.cs
@@ -33,7 +33,6 @@ public class LiveCodingConsole : ModuleRules
 			PrivateDependencyModuleNames.Add("VisualStudioSourceCodeAccess");
 		}
 
-		PublicIncludePaths.Add("Runtime/Launch/Public");
-		PrivateIncludePaths.Add("Runtime/Launch/Private");      // For LaunchEngineLoop.cpp include
+		PublicIncludePathModuleNames.Add("Launch");
 	}
 }
diff --git a/Engine/Source/Programs/Mac/DsymExporter/DsymExporter.Build.cs b/Engine/Source/Programs/Mac/DsymExporter/DsymExporter.Build.cs
--- a/Engine/Source/Programs/Mac/DsymExporter/DsymExporter.Build.cs
+++ b/Engine/Source/Programs/Mac/DsymExporter/DsymExporter.Build.cs
@@ -7,7 +7,6 @@ public class DsymExporter : ModuleRules
 	public DsymExporter( ReadOnlyTargetRules Target ) : base(Target)
 	{
 		PrivateIncludePathModuleNames.Add( "Launch" );
-		PrivateIncludePaths.Add( "Runtime/Launch/Private" );
 	
 		PrivateDependencyModuleNames.AddRange(
 			new string[] 
diff --git a/Engine/Source/Programs/Mac/UnrealAtoS/UnrealAtoS.Build.cs b/Engine/Source/Programs/Mac/UnrealAtoS/UnrealAtoS.Build.cs
--- a/Engine/Source/Programs/Mac/UnrealAtoS/UnrealAtoS.Build.cs
+++ b/Engine/Source/Programs/Mac/UnrealAtoS/UnrealAtoS.Build.cs
@@ -7,7 +7,6 @@ public class UnrealAtoS : ModuleRules
 	public UnrealAtoS( ReadOnlyTargetRules Target ) : base(Target)
 	{
 		PrivateIncludePathModuleNames.Add( "Launch" );
-		PrivateIncludePaths.Add( "Runtime/Launch/Private" );
 	
 		PrivateDependencyModuleNames.AddRange(
 			new string[] 
diff --git a/Engine/Source/Programs/Mac/UnrealEditorServices/UnrealEditorServices.Build.cs b/Engine/Source/Programs/Mac/UnrealEditorServices/UnrealEditorServices.Build.cs
--- a/Engine/Source/Programs/Mac/UnrealEditorServices/UnrealEditorServices.Build.cs
+++ b/Engine/Source/Programs/Mac/UnrealEditorServices/UnrealEditorServices.Build.cs
@@ -6,8 +6,7 @@ public class UnrealEditorServices : ModuleRules
 {
 	public UnrealEditorServices(ReadOnlyTargetRules Target) : base(Target)
 	{
-		PublicIncludePaths.Add("Runtime/Launch/Public");
-		PublicIncludePaths.Add("Runtime/Launch/Private"); // Yuck. Required for RequiredProgramMainCPPInclude.h. (Also yuck).
+		PublicIncludePathModuleNames.Add("Launch");
 
 		PrivateDependencyModuleNames.AddRange(
 			new string[] {
diff --git a/Engine/Source/Programs/ReplicationSystemTest/ReplicationSystemTest.Build.cs b/Engine/Source/Programs/ReplicationSystemTest/ReplicationSystemTest.Build.cs
--- a/Engine/Source/Programs/ReplicationSystemTest/ReplicationSystemTest.Build.cs
+++ b/Engine/Source/Programs/ReplicationSystemTest/ReplicationSystemTest.Build.cs
@@ -6,17 +6,8 @@ public class ReplicationSystemTest : ModuleRules
 {
 	public ReplicationSystemTest(ReadOnlyTargetRules Target) : base(Target)
 	{
-		PublicIncludePaths.Add("Runtime/Launch/Public");
-
-		// For LaunchEngineLoop.cpp include
-		PrivateIncludePaths.Add("Runtime/Launch/Private");
-
-		PrivateIncludePathModuleNames.AddRange(
-			new string[]
-			{
-				"DerivedDataCache",
-			}
-		);
+		PublicIncludePathModuleNames.Add("Launch");
+		PrivateIncludePathModuleNames.Add("DerivedDataCache");
 
 		PrivateDependencyModuleNames.AddRange(
 			new string[]
diff --git a/Engine/Source/Programs/ShaderCompileWorker/ShaderCompileWorker.Build.cs b/Engine/Source/Programs/ShaderCompileWorker/ShaderCompileWorker.Build.cs
--- a/Engine/Source/Programs/ShaderCompileWorker/ShaderCompileWorker.Build.cs
+++ b/Engine/Source/Programs/ShaderCompileWorker/ShaderCompileWorker.Build.cs
@@ -26,8 +26,6 @@ public class ShaderCompileWorker : ModuleRules
 				"TargetPlatform",
 			});
 
-		PrivateIncludePaths.Add("Runtime/Launch/Private");      // For LaunchEngineLoop.cpp include
-
 		// Include D3D compiler binaries
 		if (Target.Platform == UnrealTargetPlatform.Win64)
 		{
diff --git a/Engine/Source/Programs/SlateUGS/SlateUGS.Build.cs b/Engine/Source/Programs/SlateUGS/SlateUGS.Build.cs
--- a/Engine/Source/Programs/SlateUGS/SlateUGS.Build.cs
+++ b/Engine/Source/Programs/SlateUGS/SlateUGS.Build.cs
@@ -6,7 +6,7 @@ public class SlateUGS : ModuleRules
 {
 	public SlateUGS(ReadOnlyTargetRules Target) : base(Target)
 	{
-		PublicIncludePaths.Add("Runtime/Launch/Public");
+		PublicIncludePathModuleNames.Add("Launch");
 
 		PrivateDependencyModuleNames.AddRange(
 			new string[] {
@@ -24,8 +24,6 @@ public class SlateUGS : ModuleRules
 			}
 		);
 
-		PrivateIncludePaths.Add("Runtime/Launch/Private");		// For LaunchEngineLoop.cpp include
-		
 		if (Target.IsInPlatformGroup(UnrealPlatformGroup.Linux))
 		{
 			PrivateDependencyModuleNames.AddRange(
diff --git a/Engine/Source/Programs/SlateViewer/SlateViewer.Build.cs b/Engine/Source/Programs/SlateViewer/SlateViewer.Build.cs
--- a/Engine/Source/Programs/SlateViewer/SlateViewer.Build.cs
+++ b/Engine/Source/Programs/SlateViewer/SlateViewer.Build.cs
@@ -6,7 +6,7 @@ public class SlateViewer : ModuleRules
 {
 	public SlateViewer(ReadOnlyTargetRules Target) : base(Target)
 	{
-		PublicIncludePaths.Add("Runtime/Launch/Public");
+		PublicIncludePathModuleNames.Add("Launch");
 
 		PrivateDependencyModuleNames.AddRange(
 			new string[] {
@@ -44,8 +44,6 @@ public class SlateViewer : ModuleRules
 			PrivateDependencyModuleNames.Add("VisualStudioSourceCodeAccess");
 		}
 
-		PrivateIncludePaths.Add("Runtime/Launch/Private");		// For LaunchEngineLoop.cpp include
-
 		if (Target.Platform == UnrealTargetPlatform.IOS || Target.Platform == UnrealTargetPlatform.TVOS)
 		{
 			PrivateDependencyModuleNames.AddRange(
diff --git a/Engine/Source/Programs/SwitchboardListener/SwitchboardListener.Build.cs b/Engine/Source/Programs/SwitchboardListener/SwitchboardListener.Build.cs
--- a/Engine/Source/Programs/SwitchboardListener/SwitchboardListener.Build.cs
+++ b/Engine/Source/Programs/SwitchboardListener/SwitchboardListener.Build.cs
@@ -36,12 +36,6 @@ public class SwitchboardListener : ModuleRules
 			}
 		);
 
-		PrivateIncludePaths.AddRange(
-			new string[] {
-				"Runtime/Launch/Private", // for LaunchEngineLoop.cpp include
-			}
-		);
-
 		DynamicallyLoadedModuleNames.AddRange(
 			new string[] {
 				"Sockets",
diff --git a/Engine/Source/Programs/SymsLibDump/SymsLibDump.Build.cs b/Engine/Source/Programs/SymsLibDump/SymsLibDump.Build.cs
--- a/Engine/Source/Programs/SymsLibDump/SymsLibDump.Build.cs
+++ b/Engine/Source/Programs/SymsLibDump/SymsLibDump.Build.cs
@@ -6,9 +6,7 @@ public class SymsLibDump : ModuleRules
 {
 	public SymsLibDump(ReadOnlyTargetRules Target) : base(Target)
 	{
-		PublicIncludePaths.Add("Runtime/Launch/Public");
-
-		PrivateIncludePaths.Add("Runtime/Launch/Private");		// For LaunchEngineLoop.cpp include
+		PublicIncludePathModuleNames.Add("Launch");
 
 		PrivateDependencyModuleNames.AddRange(
 			new string[] {
diff --git a/Engine/Source/Programs/TestPAL/TestPAL.Build.cs b/Engine/Source/Programs/TestPAL/TestPAL.Build.cs
--- a/Engine/Source/Programs/TestPAL/TestPAL.Build.cs
+++ b/Engine/Source/Programs/TestPAL/TestPAL.Build.cs
@@ -6,10 +6,7 @@ public class TestPAL : ModuleRules
 {
 	public TestPAL(ReadOnlyTargetRules Target) : base(Target)
 	{
-		PublicIncludePaths.Add("Runtime/Launch/Public");
-
-		PrivateIncludePaths.Add("Runtime/Launch/Private");		// For LaunchEngineLoop.cpp include
-		PrivateIncludePaths.Add("Programs/TestPAL/Private");
+		PublicIncludePathModuleNames.Add("Launch");
 
 		PrivateDependencyModuleNames.AddRange(
 			new string[] {
diff --git a/Engine/Source/Programs/UnrealFrontend/UnrealFrontend.Build.cs b/Engine/Source/Programs/UnrealFrontend/UnrealFrontend.Build.cs
--- a/Engine/Source/Programs/UnrealFrontend/UnrealFrontend.Build.cs
+++ b/Engine/Source/Programs/UnrealFrontend/UnrealFrontend.Build.cs
@@ -6,17 +6,11 @@ public class UnrealFrontend : ModuleRules
 {
 	public UnrealFrontend( ReadOnlyTargetRules Target ) : base(Target)
 	{
-		PublicIncludePaths.AddRange(
-            new string[] {
-                "Runtime/Launch/Public",
-            }
-        );
+		PublicIncludePathModuleNames.Add("Launch");
 
 		PrivateIncludePaths.AddRange(
 			new string[] {
-				"Programs/UnrealFrontend/Private",
 				"Programs/UnrealFrontend/Private/Commands",
-				"Runtime/Launch/Private",					// for LaunchEngineLoop.cpp include
 			}
 		);
 
diff --git a/Engine/Source/Programs/UnrealHeaderTool/UnrealHeaderTool.Build.cs b/Engine/Source/Programs/UnrealHeaderTool/UnrealHeaderTool.Build.cs
--- a/Engine/Source/Programs/UnrealHeaderTool/UnrealHeaderTool.Build.cs
+++ b/Engine/Source/Programs/UnrealHeaderTool/UnrealHeaderTool.Build.cs
@@ -6,7 +6,15 @@ public class UnrealHeaderTool : ModuleRules
 {
 	public UnrealHeaderTool(ReadOnlyTargetRules Target) : base(Target)
 	{
-		PublicIncludePaths.Add("Runtime/Launch/Public");
+		PublicIncludePathModuleNames.Add("Launch");
+
+		PrivateIncludePathModuleNames.AddRange(
+			new string[]
+			{
+				"RHI",
+				"RigVM",
+			}
+		);
 
 		PrivateDependencyModuleNames.AddRange(
 			new string[]
@@ -17,17 +25,6 @@ public class UnrealHeaderTool : ModuleRules
 				"Projects"
 			}
 		);
-	
-		PrivateIncludePaths.AddRange(
-			new string[]
-			{
-				// For LaunchEngineLoop.cpp includes
-				"Runtime/Launch/Private",
-				"Runtime/RHI/Public",
-				"Programs/UnrealHeaderTool/Private",
-				// For RigVM Defines
-				"Runtime/RigVM/Public",
-			});
 		
 		bEnableExceptions = true;
 
diff --git a/Engine/Source/Programs/UnrealInsights/UnrealInsights.Build.cs b/Engine/Source/Programs/UnrealInsights/UnrealInsights.Build.cs
--- a/Engine/Source/Programs/UnrealInsights/UnrealInsights.Build.cs
+++ b/Engine/Source/Programs/UnrealInsights/UnrealInsights.Build.cs
@@ -6,8 +6,7 @@ public class UnrealInsights : ModuleRules
 {
 	public UnrealInsights(ReadOnlyTargetRules Target) : base(Target)
 	{
-		PublicIncludePaths.Add("Runtime/Launch/Public");
-		PrivateIncludePaths.Add("Runtime/Launch/Private"); // For LaunchEngineLoop.cpp include
+		PublicIncludePathModuleNames.Add("Launch");
 
 		PrivateDependencyModuleNames.AddRange(
 			new string[] {
diff --git a/Engine/Source/Programs/UnrealLightmass/UnrealLightmass.Build.cs b/Engine/Source/Programs/UnrealLightmass/UnrealLightmass.Build.cs
--- a/Engine/Source/Programs/UnrealLightmass/UnrealLightmass.Build.cs
+++ b/Engine/Source/Programs/UnrealLightmass/UnrealLightmass.Build.cs
@@ -7,9 +7,9 @@ public class UnrealLightmass : ModuleRules
 {
 	public UnrealLightmass(ReadOnlyTargetRules Target) : base(Target)
 	{
-		PublicIncludePaths.Add("Runtime/Launch/Public");
+		PublicIncludePathModuleNames.Add("Launch");
 
-        PrivateDependencyModuleNames.AddRange(new string[] { "Core", "CoreUObject", "zlib", "SwarmInterface", "Projects", "ApplicationCore" });
+		PrivateDependencyModuleNames.AddRange(new string[] { "Core", "CoreUObject", "zlib", "SwarmInterface", "Projects", "ApplicationCore" });
 
 		PublicDefinitions.Add("UE_LIGHTMASS=1");
 
@@ -47,7 +47,6 @@ public class UnrealLightmass : ModuleRules
 		// Warning: This means buffer overwrites will not be detected
 		bEnableBufferSecurityChecks = false;
 
-		PrivateIncludePaths.Add("Runtime/Launch/Private");		// For LaunchEngineLoop.cpp include
 		PrivateIncludePaths.Add("Programs/UnrealLightmass/Private/Launch");
 		PrivateIncludePaths.Add("Programs/UnrealLightmass/Private/ImportExport");
 		PrivateIncludePaths.Add("Programs/UnrealLightmass/Private/CPUSolver");
diff --git a/Engine/Source/Programs/UnrealMultiUserServer/UnrealMultiUserServer.Build.cs b/Engine/Source/Programs/UnrealMultiUserServer/UnrealMultiUserServer.Build.cs
--- a/Engine/Source/Programs/UnrealMultiUserServer/UnrealMultiUserServer.Build.cs
+++ b/Engine/Source/Programs/UnrealMultiUserServer/UnrealMultiUserServer.Build.cs
@@ -27,12 +27,6 @@ public class UnrealMultiUserServer : ModuleRules
 			}
 		);
 
-		PrivateIncludePaths.AddRange(
-			new string[] {
-				"Runtime/Launch/Private",           // for LaunchEngineLoop.cpp include
-			}
-		);
-
 		DynamicallyLoadedModuleNames.AddRange(
 			new string[] {
 				"ConcertSyncServer"
diff --git a/Engine/Source/Programs/UnrealMultiUserSlateServer/UnrealMultiUserSlateServer.Build.cs b/Engine/Source/Programs/UnrealMultiUserSlateServer/UnrealMultiUserSlateServer.Build.cs
--- a/Engine/Source/Programs/UnrealMultiUserSlateServer/UnrealMultiUserSlateServer.Build.cs
+++ b/Engine/Source/Programs/UnrealMultiUserSlateServer/UnrealMultiUserSlateServer.Build.cs
@@ -28,12 +28,6 @@ public class UnrealMultiUserSlateServer : ModuleRules
 			}
 		);
 
-		PrivateIncludePaths.AddRange(
-			new string[] {
-				"Runtime/Launch/Private",           // for LaunchEngineLoop.cpp include
-			}
-		);
-
 		DynamicallyLoadedModuleNames.AddRange(
 			new string[] {
 				"ConcertSyncServer",
diff --git a/Engine/Source/Programs/UnrealObjectPtrTool/UnrealObjectPtrTool.Build.cs b/Engine/Source/Programs/UnrealObjectPtrTool/UnrealObjectPtrTool.Build.cs
--- a/Engine/Source/Programs/UnrealObjectPtrTool/UnrealObjectPtrTool.Build.cs
+++ b/Engine/Source/Programs/UnrealObjectPtrTool/UnrealObjectPtrTool.Build.cs
@@ -6,10 +6,7 @@ public class UnrealObjectPtrTool : ModuleRules
 {
 	public UnrealObjectPtrTool(ReadOnlyTargetRules Target) : base(Target)
 	{
-		PublicIncludePaths.Add("Runtime/Launch/Public");
-
-		PrivateIncludePaths.Add("Runtime/Launch/Private");		// For LaunchEngineLoop.cpp include
-
+		PublicIncludePathModuleNames.Add("Launch");
 		PrivateDependencyModuleNames.Add("Core");
 		PrivateDependencyModuleNames.Add("Projects");
 	}
diff --git a/Engine/Source/Programs/UnrealPackageTool/UnrealPackageTool.Build.cs b/Engine/Source/Programs/UnrealPackageTool/UnrealPackageTool.Build.cs
--- a/Engine/Source/Programs/UnrealPackageTool/UnrealPackageTool.Build.cs
+++ b/Engine/Source/Programs/UnrealPackageTool/UnrealPackageTool.Build.cs
@@ -6,9 +6,7 @@ public class UnrealPackageTool : ModuleRules
 {
 	public UnrealPackageTool(ReadOnlyTargetRules Target) : base(Target)
 	{
-		PublicIncludePaths.Add("Runtime/Launch/Public");
-		PrivateIncludePaths.Add("Runtime/Launch/Private");
-
+		PublicIncludePathModuleNames.Add("Launch");
 		PrivateDependencyModuleNames.Add("Core");
 		PrivateDependencyModuleNames.Add("CoreUObject");
 		PrivateDependencyModuleNames.Add("Projects");
diff --git a/Engine/Source/Programs/UnrealPak/UnrealPak.Build.cs b/Engine/Source/Programs/UnrealPak/UnrealPak.Build.cs
--- a/Engine/Source/Programs/UnrealPak/UnrealPak.Build.cs
+++ b/Engine/Source/Programs/UnrealPak/UnrealPak.Build.cs
@@ -6,7 +6,7 @@ public class UnrealPak : ModuleRules
 {
 	public UnrealPak(ReadOnlyTargetRules Target) : base(Target)
 	{
-		PublicIncludePaths.Add("Runtime/Launch/Public");
+		PublicIncludePathModuleNames.Add("Launch");
 
 		PrivateDependencyModuleNames.AddRange(new string[] { 
 			"Core", 
@@ -20,8 +20,6 @@ public class UnrealPak : ModuleRules
 			"ApplicationCore" 
 		});
 
-		PrivateIncludePaths.Add("Runtime/Launch/Private"); // For LaunchEngineLoop.cpp include
-
 		PrivateIncludePathModuleNames.AddRange(new string[] {
 			"Json"
 		});
diff --git a/Engine/Source/Programs/UnrealRecoverySvc/UnrealRecoverySvc.Build.cs b/Engine/Source/Programs/UnrealRecoverySvc/UnrealRecoverySvc.Build.cs
--- a/Engine/Source/Programs/UnrealRecoverySvc/UnrealRecoverySvc.Build.cs
+++ b/Engine/Source/Programs/UnrealRecoverySvc/UnrealRecoverySvc.Build.cs
@@ -24,12 +24,5 @@ public class UnrealRecoverySvc : ModuleRules
 				"ConcertSyncServer",
 			}
 		);
-
-		PrivateIncludePaths.AddRange(
-			new string[] {
-				"Runtime/Launch/Private",           // for LaunchEngineLoop.cpp include
-			}
-		);
-
 	}
 }
diff --git a/Engine/Source/Programs/UnrealVersionSelector/UnrealVersionSelector.Build.cs b/Engine/Source/Programs/UnrealVersionSelector/UnrealVersionSelector.Build.cs
--- a/Engine/Source/Programs/UnrealVersionSelector/UnrealVersionSelector.Build.cs
+++ b/Engine/Source/Programs/UnrealVersionSelector/UnrealVersionSelector.Build.cs
@@ -6,8 +6,7 @@ public class UnrealVersionSelector : ModuleRules
 {
 	public UnrealVersionSelector(ReadOnlyTargetRules Target) : base(Target)
 	{
-		PublicIncludePaths.Add("Runtime/Launch/Public");
-		PublicIncludePaths.Add("Runtime/Launch/Private"); // Yuck. Required for RequiredProgramMainCPPInclude.h. (Also yuck).
+		PublicIncludePathModuleNames.Add("Launch");
 
 		if (Target.Platform == UnrealTargetPlatform.Linux)
 		{
diff --git a/Engine/Source/Programs/UnrealVirtualizationTool/UnrealVirtualizationTool.Build.cs b/Engine/Source/Programs/UnrealVirtualizationTool/UnrealVirtualizationTool.Build.cs
--- a/Engine/Source/Programs/UnrealVirtualizationTool/UnrealVirtualizationTool.Build.cs
+++ b/Engine/Source/Programs/UnrealVirtualizationTool/UnrealVirtualizationTool.Build.cs
@@ -19,7 +19,6 @@ public class UnrealVirtualizationTool : ModuleRules
 				"DesktopPlatform"
 			});
 
-		PublicIncludePaths.Add("Runtime/Launch/Public");
-		PrivateIncludePaths.Add("Runtime/Launch/Private");      // For LaunchEngineLoop.cpp include
+		PublicIncludePathModuleNames.Add("Launch");
 	}
 }
diff --git a/Engine/Source/Programs/ZenDashboard/ZenDashboard.Build.cs b/Engine/Source/Programs/ZenDashboard/ZenDashboard.Build.cs
--- a/Engine/Source/Programs/ZenDashboard/ZenDashboard.Build.cs
+++ b/Engine/Source/Programs/ZenDashboard/ZenDashboard.Build.cs
@@ -30,7 +30,6 @@ public class ZenDashboard : ModuleRules
 			);
 		}
 
-		PublicIncludePaths.Add("Runtime/Launch/Public");
-		PrivateIncludePaths.Add("Runtime/Launch/Private");      // For LaunchEngineLoop.cpp include
+		PublicIncludePathModuleNames.Add("Launch");
 	}
 }
