Lightmass build fix

#rnx

diff --git a/Engine/Source/Programs/UnrealLightmass/Private/Lighting/Landscape.cpp b/Engine/Source/Programs/UnrealLightmass/Private/Lighting/Landscape.cpp
--- a/Engine/Source/Programs/UnrealLightmass/Private/Lighting/Landscape.cpp
+++ b/Engine/Source/Programs/UnrealLightmass/Private/Lighting/Landscape.cpp
@@ -57,8 +57,11 @@ namespace Lightmass
 		const int32 LocalX = X - ExpandQuadsX;
 		const int32 LocalY = Y - ExpandQuadsY;
 
-		const FColor* Data = GetHeightData(X, Y);
-		OutVertex.WorldTangentZ = LandscapeDataAccess::UnpackNormal(*Data);
+		const FColor* Data = GetHeightData( X, Y );
+
+		OutVertex.WorldTangentZ.X = 2.0f / 255.f * (float)Data->B - 1.0f;
+		OutVertex.WorldTangentZ.Y = 2.0f / 255.f * (float)Data->A - 1.0f;
+		OutVertex.WorldTangentZ.Z = FMath::Sqrt(FMath::Max(1.0f - (FMath::Square(OutVertex.WorldTangentZ.X) + FMath::Square(OutVertex.WorldTangentZ.Y)), 0.f));
 		OutVertex.WorldTangentX = FVector4f(OutVertex.WorldTangentZ.Z, 0.0f, -OutVertex.WorldTangentZ.X);
 		OutVertex.WorldTangentY = OutVertex.WorldTangentZ ^ OutVertex.WorldTangentX;
 
@@ -68,8 +71,8 @@ namespace Lightmass
 		OutVertex.WorldTangentY = LtWNoScale.TransformVector(OutVertex.WorldTangentY);
 		OutVertex.WorldTangentZ = LtWNoScale.TransformVector(OutVertex.WorldTangentZ);
 
-		const float Height = LandscapeDataAccess::UnpackHeight(*Data);
-		OutVertex.WorldPosition = LocalToWorld.TransformPosition(FVector4f(LocalX, LocalY, Height));
+		const uint16 Height = (Data->R << 8) + Data->G;
+		OutVertex.WorldPosition = LocalToWorld.TransformPosition( FVector4f( LocalX, LocalY, ((float)Height - 32768.f) * LANDSCAPE_ZSCALE ) );
 		//UE_LOG(LogLightmass, Log, TEXT("%d, %d, %d, %d, %d, %d, X:%f, Y:%f, Z:%f "), SectionBaseX + LocalX - ExpandQuadsX, SectionBaseY + LocalY - ExpandQuadsY, ClampedLocalX, ClampedLocalY, SectionBaseX, SectionBaseY, WorldPos.X, WorldPos.Y, WorldPos.Z);
 
 		const int32 LightmapUVIndex = 1;
