UnrealLightmass: Fix c++20 warning operator '*': deprecated between enumerations and floating-point types

#rnx

diff --git a/Engine/Source/Programs/UnrealLightmass/Private/LightmassCore/Math/LMOctree.h b/Engine/Source/Programs/UnrealLightmass/Private/LightmassCore/Math/LMOctree.h
--- a/Engine/Source/Programs/UnrealLightmass/Private/LightmassCore/Math/LMOctree.h
+++ b/Engine/Source/Programs/UnrealLightmass/Private/LightmassCore/Math/LMOctree.h
@@ -890,7 +890,7 @@ template<typename ElementType,typename OctreeSemantics>
 TOctree<ElementType,OctreeSemantics>::TOctree(const FVector4f& InOrigin,float InExtent)
 :	RootNode(NULL)
 ,	RootNodeContext(FBoxCenterAndExtent(InOrigin,FVector4f(InExtent,InExtent,InExtent,0)), OctreeSemantics::LoosenessDenominator)
-,	MinLeafExtent(InExtent * FMath::Pow((1.0f + 1.0f / (float)OctreeSemantics::LoosenessDenominator) / 2.0f,OctreeSemantics::MaxNodeDepth))
+,	MinLeafExtent(InExtent * FMath::Pow((1.0f + 1.0f / (float)OctreeSemantics::LoosenessDenominator) / 2.0f, (float)OctreeSemantics::MaxNodeDepth))
 {
 }
 
