fix conditon error: '&&' within '||' for TextureMapping.cpp

diff --git a/Engine/Source/Programs/UnrealLightmass/Private/Lighting/TextureMapping.cpp b/Engine/Source/Programs/UnrealLightmass/Private/Lighting/TextureMapping.cpp
--- a/Engine/Source/Programs/UnrealLightmass/Private/Lighting/TextureMapping.cpp
+++ b/Engine/Source/Programs/UnrealLightmass/Private/Lighting/TextureMapping.cpp
@@ -3071,8 +3071,8 @@ void FStaticLightingSystem::CalculateIndirectLightingTextureMapping(
 				
 				NewTask->bDebugThisMapping = bDebugThisMapping
 					&& (!bDebugSelectedTaskOnly 
-						|| Scene.DebugInput.LocalX >= TaskX && Scene.DebugInput.LocalX < TaskX + CacheTaskSize
-						&& Scene.DebugInput.LocalY >= TaskY && Scene.DebugInput.LocalY < TaskY + CacheTaskSize);
+						|| (Scene.DebugInput.LocalX >= TaskX && Scene.DebugInput.LocalX < TaskX + CacheTaskSize
+						&& Scene.DebugInput.LocalY >= TaskY && Scene.DebugInput.LocalY < TaskY + CacheTaskSize));
 
 				NumTasksSubmitted++;
 
@@ -3159,8 +3159,8 @@ void FStaticLightingSystem::CalculateIndirectLightingTextureMapping(
 
 					NewTask->bDebugThisMapping = bDebugThisMapping
 						&& (!bDebugSelectedTaskOnly 
-							|| Scene.DebugInput.LocalX >= TaskX && Scene.DebugInput.LocalX < TaskX + InterpolationTaskSize
-							&& Scene.DebugInput.LocalY >= TaskY && Scene.DebugInput.LocalY < TaskY + InterpolationTaskSize);
+							|| (Scene.DebugInput.LocalX >= TaskX && Scene.DebugInput.LocalX < TaskX + InterpolationTaskSize
+							&& Scene.DebugInput.LocalY >= TaskY && Scene.DebugInput.LocalY < TaskY + InterpolationTaskSize));
 
 					NumIILTasksSubmitted++;
 					FPlatformAtomics::InterlockedIncrement(&TextureMapping->NumOutstandingInterpolationTasks);
