Ensure that PostDefault state gets set in UnrealLightmass so that lighting can be built on Mac and Linux

diff --git a/Engine/Source/Programs/UnrealLightmass/Private/Launch/UnrealLightmass.cpp b/Engine/Source/Programs/UnrealLightmass/Private/Launch/UnrealLightmass.cpp
--- a/Engine/Source/Programs/UnrealLightmass/Private/Launch/UnrealLightmass.cpp
+++ b/Engine/Source/Programs/UnrealLightmass/Private/Launch/UnrealLightmass.cpp
@@ -98,6 +98,7 @@ int LightmassMain(int argc, ANSICHAR* argv[])
 	FModuleManager::LoadModuleChecked<IMessagingModule>("Messaging");
 	FModuleManager::Get().LoadModule(TEXT("Settings"));
 	IPluginManager::Get().LoadModulesForEnabledPlugins(ELoadingPhase::PreDefault);
+	IPluginManager::Get().LoadModulesForEnabledPlugins(ELoadingPhase::PostDefault);
 
 	ON_SCOPE_EXIT
 	{
