Substrate - final renaming

diff --git a/Engine/Source/Programs/UnrealLightmass/Private/ImportExport/Material.cpp b/Engine/Source/Programs/UnrealLightmass/Private/ImportExport/Material.cpp
--- a/Engine/Source/Programs/UnrealLightmass/Private/ImportExport/Material.cpp
+++ b/Engine/Source/Programs/UnrealLightmass/Private/ImportExport/Material.cpp
@@ -5,7 +5,7 @@
 
 namespace Lightmass
 {
-	// The logic of these function needs to match StrataMaterial.cpp
+	// The logic of these function needs to match SubstrateMaterial.cpp
 	bool IsOpaqueOrMaskedBlendMode(EBlendMode BlendMode)	{ return BlendMode == BLEND_Opaque || BlendMode == BLEND_Masked; }
 	bool IsTranslucentOnlyBlendMode(EBlendMode BlendMode)	{ return BlendMode == BLEND_Translucent || BlendMode == BLEND_TranslucentColoredTransmittance || BlendMode == BLEND_TranslucentGreyTransmittance; }
 	bool IsAlphaHoldoutBlendMode(EBlendMode BlendMode)		{ return BlendMode == BLEND_AlphaHoldout; }
diff --git a/Engine/Source/Programs/UnrealLightmass/Public/MaterialExport.h b/Engine/Source/Programs/UnrealLightmass/Public/MaterialExport.h
--- a/Engine/Source/Programs/UnrealLightmass/Public/MaterialExport.h
+++ b/Engine/Source/Programs/UnrealLightmass/Public/MaterialExport.h
@@ -51,10 +51,10 @@ enum EBlendMode
     BLEND_Modulate                        =4,
 	BLEND_AlphaComposite                  =5,
 	BLEND_AlphaHoldout                    =6,
-	BLEND_TranslucentColoredTransmittance =7, /*Strata only */
+	BLEND_TranslucentColoredTransmittance =7,					/*Substrate only */
 	BLEND_MAX                             =8,
-	BLEND_TranslucentGreyTransmittance    = BLEND_Translucent, /*Strata only */
-	BLEND_ColoredTransmittanceOnly        = BLEND_Modulate, /*Strata only */
+	BLEND_TranslucentGreyTransmittance    = BLEND_Translucent,	/*Substrate only */
+	BLEND_ColoredTransmittanceOnly        = BLEND_Modulate,		/*Substrate only */
 };
 
 struct FMaterialData
