Fixed compile error caused by CL# 28346339.

diff --git a/Engine/Source/Programs/UnrealLightmass/Private/ImportExport/LightmassSwarm.cpp b/Engine/Source/Programs/UnrealLightmass/Private/ImportExport/LightmassSwarm.cpp
--- a/Engine/Source/Programs/UnrealLightmass/Private/ImportExport/LightmassSwarm.cpp
+++ b/Engine/Source/Programs/UnrealLightmass/Private/ImportExport/LightmassSwarm.cpp
@@ -288,7 +288,7 @@ void FLightmassSwarm::SendTextMessageImpl(const TCHAR* Fmt, ...)
 	{
 		FMemory::Free(Buffer);
 		Buffer = (TCHAR*) FMemory::Malloc( BufferSize * sizeof(TCHAR) );
-		GET_TYPED_VARARGS_RESULT( ElementType, Buffer, BufferSize, BufferSize-1, Fmt, Fmt, Result );
+		GET_TYPED_VARARGS_RESULT( TCHAR, Buffer, BufferSize, BufferSize-1, Fmt, Fmt, Result );
 		BufferSize *= 2;
 	};
 	Buffer[Result] = 0;
