-Make Landscape use LightMapType setting
-Added a new type of mappping, FLandscapeStaticLightingGlobalVolumeMapping (akin to FStaticLightingGlobalVolumeMapping) that makes Lightmass use the exported Landscape mesh as a surface for caching but then discards the lightmaps

diff --git a/Engine/Source/Editor/UnrealEd/Private/Lightmass/Lightmass.cpp b/Engine/Source/Editor/UnrealEd/Private/Lightmass/Lightmass.cpp
--- a/Engine/Source/Editor/UnrealEd/Private/Lightmass/Lightmass.cpp
+++ b/Engine/Source/Editor/UnrealEd/Private/Lightmass/Lightmass.cpp
@@ -656,6 +656,7 @@ void FLightmassExporter::WriteToChannel( FLightmassStatistics& Stats, FGuid& Deb
 			Scene.NumLandscapeTextureMappings = LandscapeTextureMappings.Num();
 			Scene.NumSpeedTreeMappings = 0;
 			Scene.NumVolumeMappings = VolumeMappings.Num();
+			Scene.NumLandscapeVolumeMappings = LandscapeVolumeMappings.Num();
 			Scene.NumPrecomputedVisibilityBuckets = VisibilityBucketGuids.Num();
 			Scene.NumVolumetricLightmapTasks = VolumetricLightmapTaskGuids.Num();
 			Swarm.WriteChannel( Channel, &Scene, sizeof(Scene) );
@@ -1996,6 +1997,13 @@ void FLightmassExporter::WriteMappings( int32 Channel )
 		WriteBaseTextureMappingData( Channel, VolumeMapping );
 		UpdateExportProgress();
 	}
+
+	for (int32 MappingIdx = 0; MappingIdx < LandscapeVolumeMappings.Num(); MappingIdx++)
+	{
+		const FLandscapeStaticLightingGlobalVolumeMapping* VolumeMapping = LandscapeVolumeMappings[MappingIdx];
+		WriteBaseTextureMappingData( Channel, VolumeMapping );
+		UpdateExportProgress();
+	}
 }
 
 /** Finds the GUID of the mapping that is being debugged. */
diff --git a/Engine/Source/Editor/UnrealEd/Private/Lightmass/Lightmass.h b/Engine/Source/Editor/UnrealEd/Private/Lightmass/Lightmass.h
--- a/Engine/Source/Editor/UnrealEd/Private/Lightmass/Lightmass.h
+++ b/Engine/Source/Editor/UnrealEd/Private/Lightmass/Lightmass.h
@@ -253,6 +253,7 @@ private:
 	TArray<FLandscapeStaticLightingTextureMapping*> LandscapeTextureMappings;
 
 	TArray<class FStaticLightingGlobalVolumeMapping*> VolumeMappings;
+	TArray<class FLandscapeStaticLightingGlobalVolumeMapping*> LandscapeVolumeMappings;
 
 	// materials
 	TArray<UMaterialInterface*> Materials;
@@ -281,6 +282,7 @@ private:
 	friend class FLandscapeStaticLightingMesh;
 	friend class FLandscapeStaticLightingTextureMapping;
 	friend class FStaticLightingGlobalVolumeMapping;
+	friend class FLandscapeStaticLightingGlobalVolumeMapping;
 };
 
 /** Lightmass Importer class */
diff --git a/Engine/Source/Editor/UnrealEd/Private/StaticLightingSystem/StaticLightingExport.cpp b/Engine/Source/Editor/UnrealEd/Private/StaticLightingSystem/StaticLightingExport.cpp
--- a/Engine/Source/Editor/UnrealEd/Private/StaticLightingSystem/StaticLightingExport.cpp
+++ b/Engine/Source/Editor/UnrealEd/Private/StaticLightingSystem/StaticLightingExport.cpp
@@ -166,4 +166,15 @@ void FLandscapeStaticLightingTextureMapping::ExportMapping(class FLightmassExpor
 {
 	Exporter->LandscapeTextureMappings.AddUnique(this);
 }
+
+/** 
+* Export static lighting mapping instance data to an exporter 
+* @param Exporter - export interface to process static lighting data
+**/
+void FLandscapeStaticLightingGlobalVolumeMapping::ExportMapping(class FLightmassExporter* Exporter)
+{
+	Exporter->LandscapeVolumeMappings.AddUnique(this);
+}
+
+
 #endif
diff --git a/Engine/Source/Programs/UnrealLightmass/Private/ImportExport/Importer.h b/Engine/Source/Programs/UnrealLightmass/Private/ImportExport/Importer.h
--- a/Engine/Source/Programs/UnrealLightmass/Private/ImportExport/Importer.h
+++ b/Engine/Source/Programs/UnrealLightmass/Private/ImportExport/Importer.h
@@ -104,7 +104,9 @@ public:
 	TMap<FGuid,class FStaticMesh*>&									GetStaticMeshes()			{ return StaticMeshes; }
 	TMap<FGuid,class FFluidSurfaceStaticLightingTextureMapping*>&	GetFluidMappings()			{ return FluidMappings; }
 	TMap<FGuid,class FLandscapeStaticLightingTextureMapping*>&		GetLandscapeMappings()		{ return LandscapeMappings; }
-	TMap<FGuid,class FStaticLightingGlobalVolumeMapping*>&				GetVolumeMappings()			{ return VolumeMappings; }
+	TMap<FGuid,class FStaticLightingGlobalVolumeMapping*>&			GetVolumeMappings()			{ return VolumeMappings; }
+	TMap<FGuid,class FLandscapeStaticLightingGlobalVolumeMapping*>&	GetLandscapeVolumeMappings(){ return LandscapeVolumeMappings; }
+
 	TMap<FSHAHash,class FMaterial*>&								GetMaterials()				{ return Materials; }
 
 private:
@@ -122,6 +124,7 @@ private:
 	TMap<FGuid,class FFluidSurfaceStaticLightingTextureMapping*>	FluidMappings;
 	TMap<FGuid,class FLandscapeStaticLightingTextureMapping*>		LandscapeMappings;
 	TMap<FGuid,class FStaticLightingGlobalVolumeMapping*>			VolumeMappings;
+	TMap<FGuid,class FLandscapeStaticLightingGlobalVolumeMapping*>	LandscapeVolumeMappings;
 	TMap<FSHAHash,class FMaterial*>									Materials;
 
 	float LevelScale;
diff --git a/Engine/Source/Programs/UnrealLightmass/Private/ImportExport/LightmassScene.cpp b/Engine/Source/Programs/UnrealLightmass/Private/ImportExport/LightmassScene.cpp
--- a/Engine/Source/Programs/UnrealLightmass/Private/ImportExport/LightmassScene.cpp
+++ b/Engine/Source/Programs/UnrealLightmass/Private/ImportExport/LightmassScene.cpp
@@ -61,6 +61,7 @@ FSceneFileHeader::FSceneFileHeader(const FSceneFileHeader& Other)
 	NumLandscapeTextureMappings = Other.NumLandscapeTextureMappings;
 	NumSpeedTreeMappings = Other.NumSpeedTreeMappings;
 	NumVolumeMappings = Other.NumVolumeMappings;
+	NumLandscapeVolumeMappings = Other.NumLandscapeVolumeMappings;
 	NumPortals = Other.NumPortals;
 }
 
@@ -211,6 +212,7 @@ void FScene::Import( FLightmassImporter& Importer )
 	Importer.ImportObjectArray( FluidMappings, NumFluidSurfaceTextureMappings, Importer.GetFluidMappings() );
 	Importer.ImportObjectArray( LandscapeMappings, NumLandscapeTextureMappings, Importer.GetLandscapeMappings() );
 	Importer.ImportObjectArray( VolumeMappings, NumVolumeMappings, Importer.GetVolumeMappings() );
+	Importer.ImportObjectArray( LandscapeVolumeMappings, NumLandscapeVolumeMappings, Importer.GetLandscapeVolumeMappings() );
 
 	DebugMapping = FindMappingByGuid(DebugInput.MappingGuid);
 	if (DebugMapping)
@@ -388,6 +390,15 @@ const FStaticLightingMapping* FScene::FindMappingByGuid(FGuid FindGuid) const
 		}
 	}
 
+	for (int32 i = 0; i < LandscapeVolumeMappings.Num(); i++)
+	{
+		if (LandscapeVolumeMappings[i].Guid == FindGuid)
+		{
+			return &LandscapeVolumeMappings[i];
+		}
+	}
+
+
 	return NULL;
 }
 
diff --git a/Engine/Source/Programs/UnrealLightmass/Private/ImportExport/LightmassScene.h b/Engine/Source/Programs/UnrealLightmass/Private/ImportExport/LightmassScene.h
--- a/Engine/Source/Programs/UnrealLightmass/Private/ImportExport/LightmassScene.h
+++ b/Engine/Source/Programs/UnrealLightmass/Private/ImportExport/LightmassScene.h
@@ -23,6 +23,7 @@ class FBSPSurfaceStaticLighting;
 class FFluidSurfaceStaticLightingTextureMapping;
 class FLandscapeStaticLightingTextureMapping;
 class FStaticLightingGlobalVolumeMapping;
+class FLandscapeStaticLightingGlobalVolumeMapping;
 
 /** A sample of a light's surface. */
 class FLightSurfaceSample
@@ -774,6 +775,7 @@ public:
 	TArray<FFluidSurfaceStaticLightingTextureMapping>	FluidMappings;
 	TArray<FLandscapeStaticLightingTextureMapping>		LandscapeMappings;
 	TArray<FStaticLightingGlobalVolumeMapping>			VolumeMappings;
+	TArray<FLandscapeStaticLightingGlobalVolumeMapping> LandscapeVolumeMappings;
 
 	TArray<FGuid> VisibilityBucketGuids;
 	TArray<FGuid> VolumetricLightmapTaskGuids;
diff --git a/Engine/Source/Programs/UnrealLightmass/Private/Lighting/Landscape.h b/Engine/Source/Programs/UnrealLightmass/Private/Lighting/Landscape.h
--- a/Engine/Source/Programs/UnrealLightmass/Private/Lighting/Landscape.h
+++ b/Engine/Source/Programs/UnrealLightmass/Private/Lighting/Landscape.h
@@ -46,4 +46,27 @@ namespace Lightmass
 		virtual void Import( class FLightmassImporter& Importer );
 	};
 
+	/** 
+	 * A mapping that represents an object which is going to use the global volumetric lightmap.
+	 * Hack: currently represented as a texture lightmap to Lightmass for the purposes of surface caching of light (radiosity + direct lighting).
+	 */
+	class FLandscapeStaticLightingGlobalVolumeMapping : public FLandscapeStaticLightingTextureMapping
+	{
+	public:
+
+		FLandscapeStaticLightingGlobalVolumeMapping() 
+		{}
+
+		virtual class FLandscapeStaticLightingGlobalVolumeMapping* GetLandscapeVolumeMapping() override
+		{
+			return this;
+		}
+
+		virtual const FLandscapeStaticLightingGlobalVolumeMapping* GetLandscapeVolumeMapping() const override
+		{
+			return this;
+		}
+	};
+
+
 } //namespace Lightmass
diff --git a/Engine/Source/Programs/UnrealLightmass/Private/Lighting/LightingSystem.cpp b/Engine/Source/Programs/UnrealLightmass/Private/Lighting/LightingSystem.cpp
--- a/Engine/Source/Programs/UnrealLightmass/Private/Lighting/LightingSystem.cpp
+++ b/Engine/Source/Programs/UnrealLightmass/Private/Lighting/LightingSystem.cpp
@@ -180,7 +180,7 @@ FStaticLightingSystem::FStaticLightingSystem(const FLightingBuildOptions& InOpti
 	int32 NumTriangles = 0;
 	int32 NumMappings = InScene.TextureLightingMappings.Num() +
 		InScene.FluidMappings.Num() + InScene.LandscapeMappings.Num() + InScene.BspMappings.Num();
-	int32 NumMeshInstances = InScene.BspMappings.Num() + InScene.StaticMeshInstances.Num() + InScene.VolumeMappings.Num();
+	int32 NumMeshInstances = InScene.BspMappings.Num() + InScene.StaticMeshInstances.Num() + InScene.VolumeMappings.Num() + InScene.LandscapeVolumeMappings.Num();
 	AllMappings.Reserve( NumMappings );
 	Meshes.Reserve( NumMeshInstances );
 
@@ -267,6 +267,17 @@ FStaticLightingSystem::FStaticLightingSystem(const FLightingBuildOptions& InOpti
 		}
 	}
 
+	for (int32 MappingIndex = 0; MappingIndex < InScene.LandscapeVolumeMappings.Num(); MappingIndex++)
+	{
+		FLandscapeStaticLightingGlobalVolumeMapping* Mapping = &InScene.LandscapeVolumeMappings[MappingIndex];
+		Mappings.Add(Mapping->Guid, Mapping);
+		AllMappings.Add(Mapping);
+		if (bDumpAllMappings)
+		{
+			UE_LOG(LogLightmass, Log, TEXT("\t%s"), *(Mapping->Guid.ToString()));
+		}
+	}
+
 	UE_LOG(LogLightmass, Log,  TEXT("Number of static mesh instance mappings: %d"), InScene.StaticMeshInstances.Num() );
 	for (int32 MeshIndex = 0; MeshIndex < InScene.StaticMeshInstances.Num(); MeshIndex++)
 	{
@@ -1870,6 +1881,10 @@ void TCompleteStaticLightingList<StaticLightingDataType>::ApplyAndClear(FStaticL
 		do { LocalFirstElement = FirstElement; }
 		while(FPlatformAtomics::InterlockedCompareExchangePointer((void**)&FirstElement,NULL,LocalFirstElement) != LocalFirstElement);
 
+		// Should never be null, but this satisfies static-analysis
+		if (!LocalFirstElement)
+			continue;
+
 		// Traverse the local list, count the number of entries, and find the minimum guid
 		TList<StaticLightingDataType>* PreviousElement = NULL;
 		TList<StaticLightingDataType>* MinimumElementLink = NULL;
diff --git a/Engine/Source/Programs/UnrealLightmass/Private/Lighting/Mappings.h b/Engine/Source/Programs/UnrealLightmass/Private/Lighting/Mappings.h
--- a/Engine/Source/Programs/UnrealLightmass/Private/Lighting/Mappings.h
+++ b/Engine/Source/Programs/UnrealLightmass/Private/Lighting/Mappings.h
@@ -149,6 +149,16 @@ public:
 		return NULL;
 	}
 
+	virtual const class FLandscapeStaticLightingGlobalVolumeMapping* GetLandscapeVolumeMapping()
+	{
+		return NULL;
+	}
+
+	virtual const class FLandscapeStaticLightingGlobalVolumeMapping* GetLandscapeVolumeMapping() const
+	{
+		return NULL;
+	}
+
 	/**
 	 * Returns the relative processing cost used to sort tasks from slowest to fastest.
 	 *
diff --git a/Engine/Source/Programs/UnrealLightmass/Public/SceneExport.h b/Engine/Source/Programs/UnrealLightmass/Public/SceneExport.h
--- a/Engine/Source/Programs/UnrealLightmass/Public/SceneExport.h
+++ b/Engine/Source/Programs/UnrealLightmass/Public/SceneExport.h
@@ -865,6 +865,7 @@ struct FSceneFileHeader
 	int32		NumLandscapeTextureMappings;
 	int32		NumSpeedTreeMappings;
 	int32		NumVolumeMappings;
+	int32		NumLandscapeVolumeMappings;
 	int32		NumPrecomputedVisibilityBuckets;
 	int32		NumVolumetricLightmapTasks;
 };
diff --git a/Engine/Source/Runtime/Landscape/Classes/LandscapeComponent.h b/Engine/Source/Runtime/Landscape/Classes/LandscapeComponent.h
--- a/Engine/Source/Runtime/Landscape/Classes/LandscapeComponent.h
+++ b/Engine/Source/Runtime/Landscape/Classes/LandscapeComponent.h
@@ -755,7 +755,7 @@ public:
 #endif
 	virtual void GetUsedMaterials(TArray<UMaterialInterface*>& OutMaterials, bool bGetDebugMaterials = false) const override;
 	virtual FPrimitiveSceneProxy* CreateSceneProxy() override;
-	virtual ELightMapInteractionType GetStaticLightingType() const override { return LMIT_Texture;	}
+	virtual ELightMapInteractionType GetStaticLightingType() const override;
 	virtual void GetStreamingRenderAssetInfo(FStreamingTextureLevelContext& LevelContext, TArray<FStreamingRenderAssetPrimitiveInfo>& OutStreamingRenderAssets) const override;
 	virtual bool IsPrecomputedLightingValid() const override;
 
diff --git a/Engine/Source/Runtime/Landscape/Private/LandscapeComponent.cpp b/Engine/Source/Runtime/Landscape/Private/LandscapeComponent.cpp
--- a/Engine/Source/Runtime/Landscape/Private/LandscapeComponent.cpp
+++ b/Engine/Source/Runtime/Landscape/Private/LandscapeComponent.cpp
@@ -32,6 +32,14 @@ uint32 FWeightmapLayerAllocationInfo::GetHash() const
 	return Hash;
 }
 
+ELightMapInteractionType ULandscapeComponent::GetStaticLightingType() const 
+{ 
+	if (GetLightmapType() == ELightmapType::ForceVolumetric)
+		return LMIT_GlobalVolume;
+
+	return LMIT_Texture;	
+}
+
 #if WITH_EDITOR
 
 void FLandscapeEditToolRenderData::UpdateDebugColorMaterial(const ULandscapeComponent* const Component)
diff --git a/Engine/Source/Runtime/Landscape/Private/LandscapeLight.cpp b/Engine/Source/Runtime/Landscape/Private/LandscapeLight.cpp
--- a/Engine/Source/Runtime/Landscape/Private/LandscapeLight.cpp
+++ b/Engine/Source/Runtime/Landscape/Private/LandscapeLight.cpp
@@ -45,6 +45,10 @@ FStaticLightingTextureMapping(
 {
 }
 
+FLandscapeStaticLightingGlobalVolumeMapping::FLandscapeStaticLightingGlobalVolumeMapping(ULandscapeComponent* InComponent,FStaticLightingMesh* InMesh,int32 InLightMapWidth,int32 InLightMapHeight,bool bPerformFullQualityRebuild) :	
+	FLandscapeStaticLightingTextureMapping(InComponent, InMesh, InLightMapWidth, InLightMapHeight, bPerformFullQualityRebuild)
+{}
+
 void FLandscapeStaticLightingTextureMapping::Apply(FQuantizedLightmapData* QuantizedData, const TMap<ULightComponent*,FShadowMapData2D*>& ShadowMapData, const FStaticLightingBuildContext* LightingContext)
 {
 	static const auto CVar = IConsoleManager::Get().FindTConsoleVariableDataInt(TEXT("r.VirtualTexturedLightmaps"));
@@ -692,9 +696,17 @@ void ULandscapeComponent::GetStaticLightingInfo(FStaticLightingPrimitiveInfo& Ou
 		{
 			FLandscapeStaticLightingMesh* StaticLightingMesh = new FLandscapeStaticLightingMesh(this, InRelevantLights, PatchExpandCountX, PatchExpandCountY, LightMapRatio, LightingLOD);
 			OutPrimitiveInfo.Meshes.Add(StaticLightingMesh);
-			// Create a static lighting texture mapping
-			OutPrimitiveInfo.Mappings.Add(new FLandscapeStaticLightingTextureMapping(
-				this,StaticLightingMesh,SizeX,SizeY,true));
+			if (GetLightmapType() == ELightmapType::ForceVolumetric)
+			{
+				OutPrimitiveInfo.Mappings.Add(new FLandscapeStaticLightingGlobalVolumeMapping(
+					this,StaticLightingMesh,SizeX,SizeY,true));
+			}
+			else
+			{
+				// Create a static lighting texture mapping
+				OutPrimitiveInfo.Mappings.Add(new FLandscapeStaticLightingTextureMapping(
+					this,StaticLightingMesh,SizeX,SizeY,true));
+			}
 		}
 	}
 }
diff --git a/Engine/Source/Runtime/Landscape/Public/LandscapeLight.h b/Engine/Source/Runtime/Landscape/Public/LandscapeLight.h
--- a/Engine/Source/Runtime/Landscape/Public/LandscapeLight.h
+++ b/Engine/Source/Runtime/Landscape/Public/LandscapeLight.h
@@ -46,6 +46,38 @@ private:
 	ULandscapeComponent* const LandscapeComponent;
 };
 
+class FLandscapeStaticLightingGlobalVolumeMapping  : public FLandscapeStaticLightingTextureMapping 
+{
+public:
+	LANDSCAPE_API FLandscapeStaticLightingGlobalVolumeMapping(ULandscapeComponent* InPrimitive,FStaticLightingMesh* InMesh,int32 InLightMapWidth,int32 InLightMapHeight,bool bPerformFullQualityRebuild);
+
+	virtual void Apply(struct FQuantizedLightmapData* QuantizedData, const TMap<ULightComponent*,class FShadowMapData2D*>& ShadowMapData, const FStaticLightingBuildContext* LightingContext) override
+	{
+		// Should never be processed
+		check(false);
+	}
+
+#if WITH_EDITOR
+	virtual bool DebugThisMapping() const override
+	{
+		return false;
+	}
+
+	/** 
+	 * Export static lighting mapping instance data to an exporter 
+	 * @param Exporter - export interface to process static lighting data
+	 */
+	UNREALED_API virtual void ExportMapping(class FLightmassExporter* Exporter) override;
+#endif	//WITH_EDITOR
+
+	/** Whether or not this mapping should be processed or imported */
+	virtual bool IsValidMapping() const override {return true;} 
+
+	virtual FString GetDescription() const override
+	{
+		return FString(TEXT("LandscapeVolumeMapping"));
+	}
+};
 
 
 /** Represents the triangles of a Landscape component to the static lighting system. */
diff --git a/Engine/Source/Runtime/Landscape/Public/LandscapeRender.h b/Engine/Source/Runtime/Landscape/Public/LandscapeRender.h
--- a/Engine/Source/Runtime/Landscape/Public/LandscapeRender.h
+++ b/Engine/Source/Runtime/Landscape/Public/LandscapeRender.h
@@ -710,6 +710,11 @@ class FLandscapeComponentSceneProxy : public FPrimitiveSceneProxy, public FLands
 		FLandscapeLCI(const ULandscapeComponent* InComponent, ERHIFeatureLevel::Type FeatureLevel, bool bVFRequiresPrimitiveUniformBuffer)
 			: FLightCacheInterface()
 		{
+			if (InComponent->GetLightmapType() == ELightmapType::ForceVolumetric)
+			{
+				SetGlobalVolumeLightmap(true);
+			}
+
 			const FMeshMapBuildData* MapBuildData = InComponent->GetMeshMapBuildData();
 
 			if (MapBuildData)
