Fixed Landscape not casting shadows in ForceVolumetricLightmap mode

diff --git a/Engine/Source/Editor/UnrealEd/Private/Lightmass/Lightmass.cpp b/Engine/Source/Editor/UnrealEd/Private/Lightmass/Lightmass.cpp
--- a/Engine/Source/Editor/UnrealEd/Private/Lightmass/Lightmass.cpp
+++ b/Engine/Source/Editor/UnrealEd/Private/Lightmass/Lightmass.cpp
@@ -596,6 +596,15 @@ const FStaticLightingMapping* FLightmassExporter::FindMappingByGuid(FGuid FindGu
 		}
 	}
 
+	for( int32 MappingIdx=0; MappingIdx < LandscapeVolumeMappings.Num(); MappingIdx++ )
+	{
+		const FStaticLightingMapping* CurrentMapping = LandscapeVolumeMappings[MappingIdx];
+		if (CurrentMapping->GetLightingGuid() == FindGuid)
+		{
+			return CurrentMapping;
+		}
+	}
+
 	return NULL;
 }
 
@@ -615,7 +624,7 @@ void FLightmassExporter::WriteToChannel( FLightmassStatistics& Stats, FGuid& Deb
 				DirectionalLights.Num() + PointLights.Num() + SpotLights.Num() + RectLights.Num() + SkyLights.Num() + 
 				StaticMeshes.Num() + StaticMeshLightingMeshes.Num() + StaticMeshTextureMappings.Num() + 
 				BSPSurfaceMappings.Num() + VolumeMappings.Num() + Materials.Num() + 
-				+ LandscapeLightingMeshes.Num() + LandscapeTextureMappings.Num();
+				+ LandscapeLightingMeshes.Num() + LandscapeTextureMappings.Num() + LandscapeVolumeMappings.Num();
 
 			CurrentProgress = 0;
 
@@ -2001,7 +2010,7 @@ void FLightmassExporter::WriteMappings( int32 Channel )
 	for (int32 MappingIdx = 0; MappingIdx < LandscapeVolumeMappings.Num(); MappingIdx++)
 	{
 		const FLandscapeStaticLightingGlobalVolumeMapping* VolumeMapping = LandscapeVolumeMappings[MappingIdx];
-		WriteBaseTextureMappingData( Channel, VolumeMapping );
+		WriteLandscapeMapping( Channel, VolumeMapping );
 		UpdateExportProgress();
 	}
 }
diff --git a/Engine/Source/Programs/UnrealLightmass/Private/Lighting/LightingSystem.cpp b/Engine/Source/Programs/UnrealLightmass/Private/Lighting/LightingSystem.cpp
--- a/Engine/Source/Programs/UnrealLightmass/Private/Lighting/LightingSystem.cpp
+++ b/Engine/Source/Programs/UnrealLightmass/Private/Lighting/LightingSystem.cpp
@@ -179,8 +179,8 @@ FStaticLightingSystem::FStaticLightingSystem(const FLightingBuildOptions& InOpti
 	int32 NumVertices = 0;
 	int32 NumTriangles = 0;
 	int32 NumMappings = InScene.TextureLightingMappings.Num() +
-		InScene.FluidMappings.Num() + InScene.LandscapeMappings.Num() + InScene.BspMappings.Num();
-	int32 NumMeshInstances = InScene.BspMappings.Num() + InScene.StaticMeshInstances.Num() + InScene.VolumeMappings.Num() + InScene.LandscapeVolumeMappings.Num();
+		InScene.FluidMappings.Num() + InScene.LandscapeMappings.Num() + InScene.BspMappings.Num() + InScene.LandscapeVolumeMappings.Num();
+	int32 NumMeshInstances = InScene.BspMappings.Num() + InScene.StaticMeshInstances.Num() + InScene.VolumeMappings.Num();
 	AllMappings.Reserve( NumMappings );
 	Meshes.Reserve( NumMeshInstances );
 
@@ -256,6 +256,7 @@ FStaticLightingSystem::FStaticLightingSystem(const FLightingBuildOptions& InOpti
 		}
 	}
 
+	UE_LOG(LogLightmass, Log,  TEXT("Number of volume mappings:   %d"), InScene.VolumeMappings.Num());
 	for (int32 MappingIndex = 0; MappingIndex < InScene.VolumeMappings.Num(); MappingIndex++)
 	{
 		FStaticLightingGlobalVolumeMapping* Mapping = &InScene.VolumeMappings[MappingIndex];
@@ -267,11 +268,13 @@ FStaticLightingSystem::FStaticLightingSystem(const FLightingBuildOptions& InOpti
 		}
 	}
 
+	UE_LOG(LogLightmass, Log,  TEXT("Number of landscape volume mappings:   %d"), InScene.LandscapeVolumeMappings.Num());
 	for (int32 MappingIndex = 0; MappingIndex < InScene.LandscapeVolumeMappings.Num(); MappingIndex++)
 	{
 		FLandscapeStaticLightingGlobalVolumeMapping* Mapping = &InScene.LandscapeVolumeMappings[MappingIndex];
 		Mappings.Add(Mapping->Guid, Mapping);
 		AllMappings.Add(Mapping);
+		LandscapeMappings.Add(Mapping);
 		if (bDumpAllMappings)
 		{
 			UE_LOG(LogLightmass, Log, TEXT("\t%s"), *(Mapping->Guid.ToString()));
@@ -400,14 +403,14 @@ FStaticLightingSystem::FStaticLightingSystem(const FLightingBuildOptions& InOpti
 			}
 		}
 	}
-	for (int32 MappingIndex = 0; MappingIndex < InScene.LandscapeMappings.Num(); MappingIndex++)
+	for (int32 MappingIndex = 0; MappingIndex < LandscapeMappings.Num(); MappingIndex++)
 	{
-		FLandscapeStaticLightingTextureMapping* Mapping = &InScene.LandscapeMappings[MappingIndex];
+		FStaticLightingMapping* Mapping = LandscapeMappings[MappingIndex];
 		AggregateMesh->AddMesh(Mapping->Mesh, Mapping);
 
 		if (Scene.GeneralSettings.bUseFastVoxelization)
 		{
-			if (Mapping->GetVolumeMapping() == nullptr)
+			if (Mapping->GetLandscapeVolumeMapping() == nullptr)
 			{
 				VoxelizationSurfaceAggregateMesh->AddMeshForVoxelization(Mapping->Mesh, Mapping);
 			}
@@ -549,7 +552,7 @@ FStaticLightingSystem::FStaticLightingSystem(const FLightingBuildOptions& InOpti
 	for (int32 MappingIndex = 0; MappingIndex < AllMappings.Num(); MappingIndex++)
 	{
 		FStaticLightingTextureMapping* TextureMapping = AllMappings[MappingIndex]->GetTextureMapping();
-		if (TextureMapping && !AllMappings[MappingIndex]->GetVolumeMapping())
+		if (TextureMapping && !(AllMappings[MappingIndex]->GetVolumeMapping() || AllMappings[MappingIndex]->GetLandscapeVolumeMapping()))
 		{
 			Stats.NumTexelsProcessed += TextureMapping->CachedSizeX * TextureMapping->CachedSizeY;
 		}
@@ -1796,6 +1799,7 @@ void FStaticLightingSystem::ThreadLoop(bool bIsMainThread, int32 ThreadIndex, FT
 			if(Mapping->GetTextureMapping())
 			{
 				check(!Mapping->GetVolumeMapping());
+				check(!Mapping->GetLandscapeVolumeMapping());
 				ProcessTextureMapping(Mapping->GetTextureMapping());
 				double MappingTimeEnd = FPlatformTime::Seconds();
 				ThreadStatistics.TextureMappingTime += MappingTimeEnd - MappingTimeStart;
