Lightmass: fix potentially uninitialized variable warning

diff --git a/Engine/Source/Programs/UnrealLightmass/Private/Lighting/TextureMapping.cpp b/Engine/Source/Programs/UnrealLightmass/Private/Lighting/TextureMapping.cpp
--- a/Engine/Source/Programs/UnrealLightmass/Private/Lighting/TextureMapping.cpp
+++ b/Engine/Source/Programs/UnrealLightmass/Private/Lighting/TextureMapping.cpp
@@ -1456,7 +1456,7 @@ void FStaticLightingSystem::CalculateDirectAreaLightingTextureMapping(
 				if (bLightIsInFrontOfTriangle || TextureMapping->Mesh->IsTwoSided(TexelToVertex.ElementIndex))
 				{
 					const FStaticLightingVertex CurrentVertex = TexelToVertex.GetVertex();
-					FLinearColor LightIntensity;
+					FLinearColor LightIntensity = FLinearColor::White;
 					bool bTraceShadowRays = true;
 
 					// Potentially avoid additional work below if this light has no meaningful contribution
