Simplify implementation of FLightmassSwarm::SendTextMessageImpl.

The existing logic was effectively a clone of StringBuilder's Appendf method. We now reuse Appendf here.

diff --git a/Engine/Source/Programs/UnrealLightmass/Private/ImportExport/LightmassSwarm.cpp b/Engine/Source/Programs/UnrealLightmass/Private/ImportExport/LightmassSwarm.cpp
--- a/Engine/Source/Programs/UnrealLightmass/Private/ImportExport/LightmassSwarm.cpp
+++ b/Engine/Source/Programs/UnrealLightmass/Private/ImportExport/LightmassSwarm.cpp
@@ -275,31 +275,15 @@ void FLightmassSwarm::SendAlertMessage(	NSwarm::TAlertLevel AlertLevel,
 
 
 /**
- * Sends text information to Swarm, using printf-like parameters.
+ * Sends text information to Swarm.
  */
-void FLightmassSwarm::SendTextMessageImpl(const TCHAR* Fmt, ...)
+void FLightmassSwarm::SendTextMessageImpl(const TCHAR* Buffer)
 {
-	int32		BufferSize	= 1024;
-	TCHAR*	Buffer		= NULL;
-	int32		Result		= -1;
-
-	// do the usual VARARGS shenanigans
-	while(Result == -1)
-	{
-		FMemory::Free(Buffer);
-		Buffer = (TCHAR*) FMemory::Malloc( BufferSize * sizeof(TCHAR) );
-		GET_TYPED_VARARGS_RESULT( TCHAR, Buffer, BufferSize, BufferSize-1, Fmt, Fmt, Result );
-		BufferSize *= 2;
-	};
-	Buffer[Result] = 0;
-
 	// Log it locally.
 	GLog->Log( Buffer );
 
 	// Send to Swarm.
 	SendMessage( NSwarm::FInfoMessage( Buffer ) );
-
-	FMemory::Free(Buffer);
 }
 
 /**
diff --git a/Engine/Source/Programs/UnrealLightmass/Private/ImportExport/LightmassSwarm.h b/Engine/Source/Programs/UnrealLightmass/Private/ImportExport/LightmassSwarm.h
--- a/Engine/Source/Programs/UnrealLightmass/Private/ImportExport/LightmassSwarm.h
+++ b/Engine/Source/Programs/UnrealLightmass/Private/ImportExport/LightmassSwarm.h
@@ -11,6 +11,7 @@
 #include "LMThreading.h"
 #include "Editor/SwarmInterface/Public/SwarmDefines.h"
 #include "Editor/SwarmInterface/Public/SwarmInterface.h"
+#include "Misc/StringBuilder.h"
 #include "Templates/IsValidVariadicFunctionArg.h"
 #include "Templates/IsArrayOrRefOfTypeByPredicate.h"
 #include "Traits/IsCharEncodingCompatibleWith.h"
@@ -234,14 +235,11 @@ public:
 	template <typename FmtType, typename... Types>
 	void SendTextMessage(const FmtType& Fmt, Types... Args)
 	{
-		static_assert(TIsArrayOrRefOfTypeByPredicate<FmtType, TIsCharEncodingCompatibleWithTCHAR>::Value, "Formatting string must be a TCHAR array.");
-		static_assert((TIsValidVariadicFunctionArg<Types>::Value && ...), "Invalid argument(s) passed to FString::Printf");
-
-		SendTextMessageImpl((const TCHAR*)Fmt, Args...);
+		SendTextMessageImpl(TStringBuilder<1024>().Appendf(Fmt, Args...).ToString());
 	}
 
 private:
-	void VARARGS SendTextMessageImpl(const TCHAR* Fmt, ...);
+	void SendTextMessageImpl(const TCHAR* Buffer);
 public:
 
 	/**
