UnrealLightmass: EngineIncludeOrderVersion.Latest

#rnx

diff --git a/Engine/Source/Programs/UnrealLightmass/UnrealLightmass.Target.cs b/Engine/Source/Programs/UnrealLightmass/UnrealLightmass.Target.cs
--- a/Engine/Source/Programs/UnrealLightmass/UnrealLightmass.Target.cs
+++ b/Engine/Source/Programs/UnrealLightmass/UnrealLightmass.Target.cs
@@ -11,6 +11,7 @@ public class UnrealLightmassTarget : TargetRules
 	{
 		Type = TargetType.Program;
 		LinkType = TargetLinkType.Modular;
+		IncludeOrderVersion = EngineIncludeOrderVersion.Latest;
 		AdditionalPlugins.Add("UdpMessaging");
 		LaunchModuleName = "UnrealLightmass";
 
