[UnrealLightmass] Fix non-unity compilation

#rnx

diff --git a/Engine/Source/Programs/UnrealLightmass/Private/Lighting/Radiosity.cpp b/Engine/Source/Programs/UnrealLightmass/Private/Lighting/Radiosity.cpp
--- a/Engine/Source/Programs/UnrealLightmass/Private/Lighting/Radiosity.cpp
+++ b/Engine/Source/Programs/UnrealLightmass/Private/Lighting/Radiosity.cpp
@@ -6,6 +6,7 @@
 #include "HAL/RunnableThread.h"
 #include "HAL/PlatformProcess.h"
 #include "TextureMappingSetup.h"
+#include "Misc/Compression.h"
 
 namespace Lightmass
 {
