FPlatformString and FCString: Strcpy -> Strncpy callsite changes to handle deprecation.
#rnx

diff --git a/Engine/Source/Programs/UnrealLightmass/Private/LightmassCore/Misc/LMDebug.cpp b/Engine/Source/Programs/UnrealLightmass/Private/LightmassCore/Misc/LMDebug.cpp
--- a/Engine/Source/Programs/UnrealLightmass/Private/LightmassCore/Misc/LMDebug.cpp
+++ b/Engine/Source/Programs/UnrealLightmass/Private/LightmassCore/Misc/LMDebug.cpp
@@ -95,7 +95,7 @@ void appHandleCriticalError()
 
 				TCHAR CmdLine[2048];
 				FCString::Strncpy(CmdLine, FCommandLine::Get(), UE_ARRAY_COUNT(CmdLine));
-				FCString::Strncat(CmdLine, TEXT(" -unattended"), UE_ARRAY_COUNT(CmdLine));
+				FCString::StrncatTruncateDest(CmdLine, UE_ARRAY_COUNT(CmdLine), TEXT(" -unattended"));
 				TCHAR BaseDir[260];
 				FCString::Strncpy(BaseDir, FPlatformProcess::BaseDir(), UE_ARRAY_COUNT(BaseDir));
 				TCHAR separator = 0;
