UnrealLightmass: simplify FStaticLightingSystem::IntersectLightRays()

FStaticLightingSystem::IntersectLightRays() was only ever being called
in the context of a single ray (NumRays==1). This change simplifies and
renames it to just "IntersectLightRay()" (singular) to reflect this use case,
which allows removing the Embree-only IntersectLightRays4() function
that can never be called.

This is in preparation for an upcoming upgrade of Embree where in
the transition from 2.x to 3.x, callback functions that took a particular
packet size of rays (4, 8, 16, etc.) were removed in favor of a single
callback with a variable packet size.

diff --git a/Engine/Source/Programs/UnrealLightmass/Private/Lighting/Collision.h b/Engine/Source/Programs/UnrealLightmass/Private/Lighting/Collision.h
--- a/Engine/Source/Programs/UnrealLightmass/Private/Lighting/Collision.h
+++ b/Engine/Source/Programs/UnrealLightmass/Private/Lighting/Collision.h
@@ -175,17 +175,6 @@ public:
 		class FCoherentRayCache& CoherentRayCache,
 		FLightRayIntersection& Intersection) const = 0;
 
-	virtual void IntersectLightRays4(
-		const FLightRay* LightRays,
-		bool bFindClosestIntersection,
-		bool bCalculateTransmission,
-		bool bDirectShadowingRay,
-		FCoherentRayCache& CoherentRayCache,
-		FLightRayIntersection* ClosestIntersections) const
-	{
-		check(0);
-	}
-
 	FBox3f GetBounds() const;
 
 	/** The total surface area of everything in the aggregate mesh */
diff --git a/Engine/Source/Programs/UnrealLightmass/Private/Lighting/Embree.cpp b/Engine/Source/Programs/UnrealLightmass/Private/Lighting/Embree.cpp
--- a/Engine/Source/Programs/UnrealLightmass/Private/Lighting/Embree.cpp
+++ b/Engine/Source/Programs/UnrealLightmass/Private/Lighting/Embree.cpp
@@ -827,77 +827,6 @@ bool FEmbreeAggregateMesh::IntersectLightRay(
 	return ClosestIntersection.bIntersects;
 }
 
-void FEmbreeAggregateMesh::IntersectLightRays4(
-	const FLightRay* LightRays,
-	bool bFindClosestIntersection,
-	bool bCalculateTransmission,
-	bool bDirectShadowingRay,
-	FCoherentRayCache& CoherentRayCache,
-	FLightRayIntersection* ClosestIntersections) const
-{
-	LIGHTINGSTAT(FScopedRDTSCTimer RayTraceTimer(bFindClosestIntersection ? CoherentRayCache.FirstHitRayTraceTime : CoherentRayCache.BooleanRayTraceTime);)
-	bFindClosestIntersection ? CoherentRayCache.NumFirstHitRaysTraced+=4 : CoherentRayCache.NumBooleanRaysTraced+=4;
-	// Calculating transmission requires finding the closest intersection for now
-	//@todo - allow boolean visibility tests while calculating transmission
-	checkSlow(!bCalculateTransmission || bFindClosestIntersection);
-
-	FEmbreeRay4 EmbreeRay(
-		LightRays[0].Mesh, 
-		LightRays[0].Mapping ? LightRays[0].Mapping->Mesh : NULL,
-		LightRays[0].TraceFlags, 
-		bFindClosestIntersection, 
-		bCalculateTransmission, 
-		bDirectShadowingRay);
-
-	for (int32 i = 0; i < 4; i++)
-	{
-		ClosestIntersections[i].bIntersects = false;
-		EmbreeRay.orgx[i] = LightRays[i].Start.X;
-		EmbreeRay.orgy[i] = LightRays[i].Start.Y;
-		EmbreeRay.orgz[i] = LightRays[i].Start.Z;
-		EmbreeRay.dirx[i] = LightRays[i].Direction.X;
-		EmbreeRay.diry[i] = LightRays[i].Direction.Y;
-		EmbreeRay.dirz[i] = LightRays[i].Direction.Z;
-		EmbreeRay.tnear[i] = 0;
-		EmbreeRay.tfar[i] = LightRays[i].Length;
-	}
-
-	// We use geomID instead
-	alignas(16) int32 UnusedValidMask[4] = { -1, -1, -1, -1 };
-
-	if (bFindClosestIntersection)
-	{
-		rtcIntersect4(UnusedValidMask, EmbreeScene, EmbreeRay);
-	}
-	else
-	{
-		rtcOccluded4(UnusedValidMask, EmbreeScene, EmbreeRay);
-	}
-
-	for (int32 i = 0; i < 4; i++)
-	{
-		if (EmbreeRay.geomID[i] != -1 && EmbreeRay.primID[i] != -1)
-		{
-			const FEmbreeGeometry& Geo = *(FEmbreeGeometry*)rtcGetUserData(EmbreeScene, EmbreeRay.geomID[i]);
-
-			FMinimalStaticLightingVertex EmbreeVertex;
-			EmbreeVertex.WorldPosition = LightRays[i].Start + LightRays[i].Direction * EmbreeRay.tfar[i];
-			EmbreeVertex.WorldTangentZ = FVector3f(EmbreeRay.Ngx[i], EmbreeRay.Ngy[i], EmbreeRay.Ngz[i]).GetSafeNormal();
-
-			EmbreeVertex.TextureCoordinates[0] = EmbreeRay.TextureCoordinates[i];
-			EmbreeVertex.TextureCoordinates[1] = EmbreeRay.LightmapCoordinates[i];
-
-			ClosestIntersections[i] = FLightRayIntersection(true, EmbreeVertex, Geo.Mesh, Geo.Mapping, EmbreeRay.ElementIndex[i]);
-
-			EmbreeRay.TransmissionAcc[i].Resolve(ClosestIntersections[i].Transmission, EmbreeRay.tfar[i]);
-		}
-		else
-		{
-			EmbreeRay.TransmissionAcc[i].Resolve(ClosestIntersections[i].Transmission);
-		}
-	}
-}
-
 FEmbreeVerifyAggregateMesh::FEmbreeVerifyAggregateMesh(const FScene& InScene) :
 	Super(InScene),
 	DefaultAggregate(InScene),
diff --git a/Engine/Source/Programs/UnrealLightmass/Private/Lighting/Embree.h b/Engine/Source/Programs/UnrealLightmass/Private/Lighting/Embree.h
--- a/Engine/Source/Programs/UnrealLightmass/Private/Lighting/Embree.h
+++ b/Engine/Source/Programs/UnrealLightmass/Private/Lighting/Embree.h
@@ -210,14 +210,6 @@ namespace Lightmass
 			class FCoherentRayCache& CoherentRayCache,
 			FLightRayIntersection& Intersection) const override;
 
-		virtual void IntersectLightRays4(
-			const FLightRay* LightRays,
-			bool bFindClosestIntersection,
-			bool bCalculateTransmission,
-			bool bDirectShadowingRay,
-			FCoherentRayCache& CoherentRayCache,
-			FLightRayIntersection* ClosestIntersections) const override;
-
 		TArray<const FStaticLightingMapping*> StaticMeshInstancesToMappings;
 
 	private:
diff --git a/Engine/Source/Programs/UnrealLightmass/Private/Lighting/FinalGather.cpp b/Engine/Source/Programs/UnrealLightmass/Private/Lighting/FinalGather.cpp
--- a/Engine/Source/Programs/UnrealLightmass/Private/Lighting/FinalGather.cpp
+++ b/Engine/Source/Programs/UnrealLightmass/Private/Lighting/FinalGather.cpp
@@ -304,69 +304,54 @@ FLinearColor FStaticLightingSystem::CalculateExitantRadiance(
 	return AccumulatedRadiance;
 }
 
-void FStaticLightingSystem::IntersectLightRays(
+void FStaticLightingSystem::IntersectLightRay(
 	const FStaticLightingMapping* Mapping,
 	const FFullStaticLightingVertex& Vertex,
 	float SampleRadius,
-	int32 NumRays,
-	const FVector4f* WorldPathDirections,
-	const FVector4f* TangentPathDirections,
+	const FVector4f& WorldPathDirection,
+	const FVector4f& TangentPathDirection,
 	EFinalGatherRayBiasMode RayBiasMode,
 	FStaticLightingMappingContext& MappingContext,
-	FLightRay* OutLightRays,
-	FLightRayIntersection* OutLightRayIntersections) const
+	FLightRay& OutLightRay,
+	FLightRayIntersection& OutLightRayIntersection) const
 {
-	for (int32 RayIndex = 0; RayIndex < NumRays; RayIndex++)
+	FVector4f SampleOffset(0,0,0);
+	if (GeneralSettings.bAccountForTexelSize)
 	{
-		const FVector4f WorldPathDirection = WorldPathDirections[RayIndex];
-		const FVector4f TangentPathDirection = TangentPathDirections[RayIndex];
-
-		FVector4f SampleOffset(0,0,0);
-		if (GeneralSettings.bAccountForTexelSize)
-		{
-			// Offset the sample's starting point in the tangent XY plane based on the sample's area of influence. 
-			// This is particularly effective for large texels with high variance in the incoming radiance over the area of the texel.
-			SampleOffset = Vertex.WorldTangentX * TangentPathDirection.X * SampleRadius * SceneConstants.VisibilityTangentOffsetSampleRadiusScale
-				+ Vertex.WorldTangentY * TangentPathDirection.Y * SampleRadius * SceneConstants.VisibilityTangentOffsetSampleRadiusScale;
-			
-			// Experiment to distribute the starting position over the area of the texel to anti-alias, causes incorrect shadowing at intersections though
-			//@todo - use consistent sample set between irradiance cache samples
-			//const FVector2f DiskPosition = GetUniformUnitDiskPosition(RandomStream);
-			//SampleOffset = Vertex.WorldTangentX * DiskPosition.X * SampleRadius * .5f + Vertex.WorldTangentY * DiskPosition.Y * SampleRadius * .5f;
-		}
-
-		const float RayStartNormalBiasScale = RayBiasMode == RBM_ConstantNormalOffset
-			? SceneConstants.VisibilityNormalOffsetSampleRadiusScale 
-			: (SceneConstants.VisibilityTangentOffsetSampleRadiusScale * TangentPathDirection.Z);
-
-		// Apply various offsets to the start of the ray.
-		// The offset along the ray direction is to avoid incorrect self-intersection due to floating point precision.
-		// The offset along the normal is to push self-intersection patterns (like triangle shape) on highly curved surfaces onto the backfaces.
-		FVector3f RayStart = Vertex.WorldPosition
-			+ WorldPathDirection * SceneConstants.VisibilityRayOffsetDistance
-			+ Vertex.WorldTangentZ * RayStartNormalBiasScale * SampleRadius
-			+ SampleOffset;
-
-		OutLightRays[RayIndex] = FLightRay(
-			RayStart,
-			Vertex.WorldPosition + WorldPathDirection * MaxRayDistance,
-			Mapping,
-			NULL
-			);
+		// Offset the sample's starting point in the tangent XY plane based on the sample's area of influence. 
+		// This is particularly effective for large texels with high variance in the incoming radiance over the area of the texel.
+		SampleOffset = Vertex.WorldTangentX * TangentPathDirection.X * SampleRadius * SceneConstants.VisibilityTangentOffsetSampleRadiusScale
+			+ Vertex.WorldTangentY * TangentPathDirection.Y * SampleRadius * SceneConstants.VisibilityTangentOffsetSampleRadiusScale;
+		
+		// Experiment to distribute the starting position over the area of the texel to anti-alias, causes incorrect shadowing at intersections though
+		//@todo - use consistent sample set between irradiance cache samples
+		//const FVector2f DiskPosition = GetUniformUnitDiskPosition(RandomStream);
+		//SampleOffset = Vertex.WorldTangentX * DiskPosition.X * SampleRadius * .5f + Vertex.WorldTangentY * DiskPosition.Y * SampleRadius * .5f;
 	}
 
-	MappingContext.Stats.NumFirstBounceRaysTraced += NumRays;
+	const float RayStartNormalBiasScale = RayBiasMode == RBM_ConstantNormalOffset
+		? SceneConstants.VisibilityNormalOffsetSampleRadiusScale 
+		: (SceneConstants.VisibilityTangentOffsetSampleRadiusScale * TangentPathDirection.Z);
+
+	// Apply various offsets to the start of the ray.
+	// The offset along the ray direction is to avoid incorrect self-intersection due to floating point precision.
+	// The offset along the normal is to push self-intersection patterns (like triangle shape) on highly curved surfaces onto the backfaces.
+	FVector3f RayStart = Vertex.WorldPosition
+		+ WorldPathDirection * SceneConstants.VisibilityRayOffsetDistance
+		+ Vertex.WorldTangentZ * RayStartNormalBiasScale * SampleRadius
+		+ SampleOffset;
+
+	OutLightRay = FLightRay(
+		RayStart,
+		Vertex.WorldPosition + WorldPathDirection * MaxRayDistance,
+		Mapping,
+		NULL
+		);
+
+	MappingContext.Stats.NumFirstBounceRaysTraced += 1;
 	const float LastRayTraceTime = MappingContext.RayCache.FirstHitRayTraceTime;
 
-	if (NumRays == 1)
-	{
-		AggregateMesh->IntersectLightRay(OutLightRays[0], true, false, false, MappingContext.RayCache, OutLightRayIntersections[0]);
-	}
-	else
-	{
-		checkSlow(NumRays == 4);
-		AggregateMesh->IntersectLightRays4(OutLightRays, true, false, false, MappingContext.RayCache, OutLightRayIntersections);
-	}
+	AggregateMesh->IntersectLightRay(OutLightRay, true, false, false, MappingContext.RayCache, OutLightRayIntersection);
 	
 	MappingContext.Stats.FirstBounceRayTraceTime += MappingContext.RayCache.FirstHitRayTraceTime - LastRayTraceTime;
 }
@@ -933,10 +918,10 @@ public:
 			// Swap input and output for the next step
 			Swap(CurrentNodesToRefine, NextNodesToRefine);
 
-			FVector4f WorldPathDirections[4];
-			FVector4f TangentPathDirections[4];
-			FLightRay LightRays[4];
-			FLightRayIntersection LightRayIntersections[4];
+			FVector4f WorldPathDirection;
+			FVector4f TangentPathDirection;
+			FLightRay LightRay;
+			FLightRayIntersection LightRayIntersection;
 
 			FStaticLightingMappingStats& Stats = MappingContext.Stats;
 
@@ -987,19 +972,18 @@ public:
 
 							const FVector4f SampleDirection = UniformSampleHemisphere(Fraction1, Fraction2);
 
-							Vertex.ComputePathDirections(SampleDirection, WorldPathDirections[0], TangentPathDirections[0]);
+							Vertex.ComputePathDirections(SampleDirection, WorldPathDirection, TangentPathDirection);
 
-							LightingSystem.IntersectLightRays(
+							LightingSystem.IntersectLightRay(
 								Mapping,
 								Vertex,
 								SampleRadius,
-								1,
-								WorldPathDirections,
-								TangentPathDirections,
+								WorldPathDirection,
+								TangentPathDirection,
 								RayBiasMode,
 								MappingContext,
-								LightRays,
-								LightRayIntersections);
+								LightRay,
+								LightRayIntersection);
 
 							FVector3f UnoccludedSkyVector;
 							FLinearColor StationarySkyLighting;
@@ -1009,10 +993,10 @@ public:
 							const FLinearColor SubsampleLighting = LightingSystem.FinalGatherSample(
 								Mapping,
 								Vertex,
-								WorldPathDirections[0],
-								TangentPathDirections[0],
-								LightRays[0],
-								LightRayIntersections[0],
+								WorldPathDirection,
+								TangentPathDirection,
+								LightRay,
+								LightRayIntersection,
 								NodeContext.SolidAngle,
 								BounceNumber,
 								GatherClassification,
@@ -1165,10 +1149,8 @@ SampleType FStaticLightingSystem::IncomingRadianceAdaptive(
 
 	const float BaseGridSolidAngle = 2 * PI / (float)UniformHemisphereSamples.Num();
 
-	FVector4f WorldPathDirections[4];
-	FVector4f TangentPathDirections[4];
-	FLightRay LightRays[4];
-	FLightRayIntersection LightRayIntersections[4];
+	FLightRay LightRay;
+	FLightRayIntersection LightRayIntersection;
 
 	// Initialize the root level of the refinement grid with lighting values
 	for (int32 ThetaIndex = 0; ThetaIndex < NumThetaSteps; ThetaIndex++)
@@ -1178,19 +1160,20 @@ SampleType FStaticLightingSystem::IncomingRadianceAdaptive(
 			const int32 SampleIndex = ThetaIndex * NumPhiSteps + PhiIndex;
 			const FVector4f TriangleTangentPathDirection = UniformHemisphereSamples[SampleIndex];
 			
-			Vertex.ComputePathDirections(TriangleTangentPathDirection, WorldPathDirections[0], TangentPathDirections[0]);
+			FVector4f WorldPathDirection;
+			FVector4f TangentPathDirection;
+			Vertex.ComputePathDirections(TriangleTangentPathDirection, WorldPathDirection, TangentPathDirection);
 
-			IntersectLightRays(
+			IntersectLightRay(
 				Mapping,
 				Vertex,
 				SampleRadius,
-				1,
-				WorldPathDirections,
-				TangentPathDirections,
+				WorldPathDirection,
+				TangentPathDirection,
 				RayBiasMode,
 				MappingContext,
-				LightRays,
-				LightRayIntersections);
+				LightRay,
+				LightRayIntersection);
 
 			FVector3f UnoccludedSkyVector;
 			FLinearColor StationarySkyLighting;
@@ -1200,10 +1183,10 @@ SampleType FStaticLightingSystem::IncomingRadianceAdaptive(
 			const FLinearColor Radiance = FinalGatherSample(
 				Mapping,
 				Vertex,
-				WorldPathDirections[0],
-				TangentPathDirections[0],
-				LightRays[0],
-				LightRayIntersections[0],
+				WorldPathDirection,
+				TangentPathDirection,
+				LightRay,
+				LightRayIntersection,
 				BaseGridSolidAngle,
 				BounceNumber,
 				GatherClassification,
diff --git a/Engine/Source/Programs/UnrealLightmass/Private/Lighting/LightingSystem.h b/Engine/Source/Programs/UnrealLightmass/Private/Lighting/LightingSystem.h
--- a/Engine/Source/Programs/UnrealLightmass/Private/Lighting/LightingSystem.h
+++ b/Engine/Source/Programs/UnrealLightmass/Private/Lighting/LightingSystem.h
@@ -2298,17 +2298,16 @@ private:
 		FStaticLightingMappingContext& MappingContext,
 		bool bDebugThisTexel) const;
 
-	void IntersectLightRays(
+	void IntersectLightRay(
 		const FStaticLightingMapping* Mapping,
 		const FFullStaticLightingVertex& Vertex,
 		float SampleRadius,
-		int32 NumRays,
-		const FVector4f* WorldPathDirections,
-		const FVector4f* TangentPathDirections,
+		const FVector4f& WorldPathDirection,
+		const FVector4f& TangentPathDirection,
 		EFinalGatherRayBiasMode RayBiasMode,
 		FStaticLightingMappingContext& MappingContext,
-		FLightRay* OutLightRays,
-		FLightRayIntersection* OutLightRayIntersections) const;
+		FLightRay& OutLightRay,
+		FLightRayIntersection& OutLightRayIntersection) const;
 
 	/** Takes a final gather sample in the given tangent space direction. */
 	FLinearColor FinalGatherSample(
