UnrealLightmass: fix Embree memory monitor function signature and API call

At the moment, UnrealLightmass is back on Embree 2.7.0 while the rest of
the engine is on Embree 3.12.2. There was a change in the signature of the
memory monitor function type and the API call to set it between these two
versions of Embree, so this change switches the usage in UnrealLightmass
back to the 2.7.0 version.

diff --git a/Engine/Source/Programs/UnrealLightmass/Private/ImportExport/LightmassScene.cpp b/Engine/Source/Programs/UnrealLightmass/Private/ImportExport/LightmassScene.cpp
--- a/Engine/Source/Programs/UnrealLightmass/Private/ImportExport/LightmassScene.cpp
+++ b/Engine/Source/Programs/UnrealLightmass/Private/ImportExport/LightmassScene.cpp
@@ -86,8 +86,7 @@ FScene::~FScene()
 #if USE_EMBREE
 LLM_DECLARE_TAG(Embree);
 
-// userPtr is provided during callback registration, it points to the associated FEmbreeScene
-static bool EmbreeMemoryMonitor(void* userPtr, ssize_t bytes, bool post)
+static bool EmbreeMemoryMonitor(const ssize_t bytes, const bool post)
 {
 	LLM_SCOPE_BYTAG(Embree);
 	LLM_IF_ENABLED(FLowLevelMemTracker::Get().OnLowLevelChangeInMemoryUse(ELLMTracker::Default, static_cast<int64>(bytes)));
@@ -109,7 +108,7 @@ void FScene::Import( FLightmassImporter& Importer )
 	{
 		EmbreeDevice = rtcNewDevice(NULL);
 		check(rtcDeviceGetError(EmbreeDevice) == RTC_NO_ERROR);
-		LLM_IF_ENABLED(rtcSetDeviceMemoryMonitorFunction(EmbreeDevice, EmbreeMemoryMonitor, NULL));
+		LLM_IF_ENABLED(rtcDeviceSetMemoryMonitorFunction(EmbreeDevice, EmbreeMemoryMonitor));
 		bVerifyEmbree = TempHeader.GeneralSettings.bVerifyEmbree;
 	}
 #endif
