UnrealLightmass: re-name LLM LightmassScene Embree memory monitor function to avoid collision

diff --git a/Engine/Source/Programs/UnrealLightmass/Private/ImportExport/LightmassScene.cpp b/Engine/Source/Programs/UnrealLightmass/Private/ImportExport/LightmassScene.cpp
--- a/Engine/Source/Programs/UnrealLightmass/Private/ImportExport/LightmassScene.cpp
+++ b/Engine/Source/Programs/UnrealLightmass/Private/ImportExport/LightmassScene.cpp
@@ -86,7 +86,7 @@ FScene::~FScene()
 #if USE_EMBREE
 LLM_DECLARE_TAG(Embree);
 
-static bool EmbreeMemoryMonitor(const ssize_t bytes, const bool post)
+static bool EmbreeLlmMemoryMonitor(const ssize_t bytes, const bool post)
 {
 	LLM_SCOPE_BYTAG(Embree);
 	LLM_IF_ENABLED(FLowLevelMemTracker::Get().OnLowLevelChangeInMemoryUse(ELLMTracker::Default, static_cast<int64>(bytes)));
@@ -108,7 +108,7 @@ void FScene::Import( FLightmassImporter& Importer )
 	{
 		EmbreeDevice = rtcNewDevice(NULL);
 		check(rtcDeviceGetError(EmbreeDevice) == RTC_NO_ERROR);
-		LLM_IF_ENABLED(rtcDeviceSetMemoryMonitorFunction(EmbreeDevice, EmbreeMemoryMonitor));
+		LLM_IF_ENABLED(rtcDeviceSetMemoryMonitorFunction(EmbreeDevice, EmbreeLlmMemoryMonitor));
 		bVerifyEmbree = TempHeader.GeneralSettings.bVerifyEmbree;
 	}
 #endif
