UnrealLightmass: enable building against Embree 3 and 4

This adds preprocessor protected code to use the appropriate Embree
API based on the major version number (2, 3, or 4) as determined by the
definition USE_EMBREE_MAJOR_VERSION.

Besides the change in header file paths, the notable changes between
Embree 2 and 3 are:
- Filter functions now receive a struct containing arguments, and
potentially a variable size packet of rays/hits, although since we
use rtcIntersect1() and rtcOccluded1(), we are guaranteed to
receive single-ray packets. A "valid" array in the args is expected
to be used to discard/invalidate rays.
- Embree 3 and 4 no longer support carrying additional data along
with rays by extending RTCRay, so instead we put that data into
the intersect/ray query context, along with storage for the outputs
of the filter function.
- The ray and hit data previously all contained in an RTCRay
were broken out into separate RTCRay and RTCHit structs.
- The API for creating new geometry and setting up and populating
data buffers was changed. See the link below for more detail on the
changes. Geometry also needs to be committed as part of these changes.
- The direction of the geometric normal was flipped in Embree 3 to follow
the right-hand rule.
- Other minor function name and callback signature changes.

For more detail, see here:
[REDACTED_PROFILE_URL]/embree/tree/v3.12.2#upgrading-from-embree-2-to-embree-3

The notable change between Embree 3 and 4 is:
- RTCIntersectContext was renamed to RTCRayQueryContext,
and it was also made part of the optional
RTCIntersectArguments/RTCOccludedArguments parameters
to rtcIntersect and rtcOccluded functions.

For more detail, see here:
[REDACTED_PROFILE_URL]/embree/tree/v4.3.3#upgrading-from-embree-3-to-embree-4

Embree 2.7.0 is still the version being used by UnrealLightmass,
so there is no functional change introduced here.

This is in preparation for upgrading to Embree version 4.3.3.

diff --git a/Engine/Source/Programs/UnrealLightmass/Private/ImportExport/LightmassScene.cpp b/Engine/Source/Programs/UnrealLightmass/Private/ImportExport/LightmassScene.cpp
--- a/Engine/Source/Programs/UnrealLightmass/Private/ImportExport/LightmassScene.cpp
+++ b/Engine/Source/Programs/UnrealLightmass/Private/ImportExport/LightmassScene.cpp
@@ -79,14 +79,22 @@ FScene::FScene() :
 FScene::~FScene()
 {
 #if USE_EMBREE
+#if USE_EMBREE_MAJOR_VERSION >= 3
+	rtcReleaseDevice(EmbreeDevice);
+#else
 	rtcDeleteDevice(EmbreeDevice);
+#endif // USE_EMBREE_MAJOR_VERSION >= 3
 #endif
 }
 
 #if USE_EMBREE
 LLM_DECLARE_TAG(Embree);
 
+#if USE_EMBREE_MAJOR_VERSION >= 3
+static bool EmbreeLlmMemoryMonitor(void* ptr, ssize_t bytes, bool post)
+#else
 static bool EmbreeLlmMemoryMonitor(const ssize_t bytes, const bool post)
+#endif // USE_EMBREE_MAJOR_VERSION >= 3
 {
 	LLM_SCOPE_BYTAG(Embree);
 	LLM_IF_ENABLED(FLowLevelMemTracker::Get().OnLowLevelChangeInMemoryUse(ELLMTracker::Default, static_cast<int64>(bytes)));
@@ -107,8 +115,13 @@ void FScene::Import( FLightmassImporter& Importer )
 	if (TempHeader.GeneralSettings.bUseEmbree)
 	{
 		EmbreeDevice = rtcNewDevice(NULL);
+#if USE_EMBREE_MAJOR_VERSION >= 3
+		check(rtcGetDeviceError(EmbreeDevice) == RTC_ERROR_NONE);
+		LLM_IF_ENABLED(rtcSetDeviceMemoryMonitorFunction(EmbreeDevice, EmbreeLlmMemoryMonitor, NULL));
+#else
 		check(rtcDeviceGetError(EmbreeDevice) == RTC_NO_ERROR);
 		LLM_IF_ENABLED(rtcDeviceSetMemoryMonitorFunction(EmbreeDevice, EmbreeLlmMemoryMonitor));
+#endif // USE_EMBREE_MAJOR_VERSION >= 3
 		bVerifyEmbree = TempHeader.GeneralSettings.bVerifyEmbree;
 	}
 #endif
diff --git a/Engine/Source/Programs/UnrealLightmass/Private/Lighting/Embree.cpp b/Engine/Source/Programs/UnrealLightmass/Private/Lighting/Embree.cpp
--- a/Engine/Source/Programs/UnrealLightmass/Private/Lighting/Embree.cpp
+++ b/Engine/Source/Programs/UnrealLightmass/Private/Lighting/Embree.cpp
@@ -46,6 +46,30 @@ void FEmbreeTransmissionAccumulator::Resolve(FLinearColor& FinalColor)
 	}
 }
 
+#if USE_EMBREE_MAJOR_VERSION >= 3
+FEmbreeContext::FEmbreeContext(
+	const FStaticLightingMesh* InShadowMesh,
+	const FStaticLightingMesh* InMappingMesh,
+	uint32 InTraceFlags,
+	bool InFindClosestIntersection,
+	bool InCalculateTransmission,
+	bool InDirectShadowingRay
+	) :
+	ShadowMesh(InShadowMesh),
+	MappingMesh(InMappingMesh),
+	TraceFlags(InTraceFlags),
+	bFindClosestIntersection(InFindClosestIntersection),
+	bCalculateTransmission(InCalculateTransmission),
+	bDirectShadowingRay(InDirectShadowingRay),
+	bStaticAndOpaqueOnly((TraceFlags & LIGHTRAY_STATIC_AND_OPAQUEONLY) != 0),
+	bTwoSidedCollision(!InDirectShadowingRay),
+	bFlipSidedness((TraceFlags & LIGHTRAY_FLIP_SIDEDNESS) != 0),
+	ElementIndex(-1),
+	TextureCoordinates(0, 0),
+	LightmapCoordinates(0, 0)
+{
+}
+#else
 FEmbreeRay::FEmbreeRay(
 	const FStaticLightingMesh* InShadowMesh, 
 	const FStaticLightingMesh* InMappingMesh, 
@@ -107,17 +131,44 @@ FEmbreeRay4::FEmbreeRay4(
 		LightmapCoordinates[i] = FVector2f(0, 0);
 	}
 }
+#endif // USE_EMBREE_MAJOR_VERSION >= 3
 
 struct FEmbreeFilterProcessor
 {
-	FEmbreeFilterProcessor(FEmbreeRay& InRay, const FEmbreeGeometry& InGeo) : Ray(InRay), Geo(InGeo), bCoordsDirty(true) 
+#if USE_EMBREE_MAJOR_VERSION >= 3
+	FEmbreeFilterProcessor(
+		FEmbreeContext& InEmbreeContext,
+		RTCRay& InRay,
+		RTCHit& InHit,
+		int*& InValidMask,
+		const FEmbreeGeometry& InGeo
+	) :
+	EmbreeContext(InEmbreeContext),
+	EmbreeRay(InRay),
+	EmbreeHit(InHit),
+	ValidMask(InValidMask),
+	Geo(InGeo),
+	bCoordsDirty(true)
+#else
+	FEmbreeFilterProcessor(FEmbreeRay& InRay, const FEmbreeGeometry& InGeo) : Ray(InRay), Geo(InGeo), bCoordsDirty(true)
+#endif // USE_EMBREE_MAJOR_VERSION >= 3
 	{
-		Desc = Geo.TriangleDescs[Ray.primID];
 		Mesh = Geo.Mesh;
+
+#if USE_EMBREE_MAJOR_VERSION >= 3
+		Desc = Geo.TriangleDescs[EmbreeHit.primID];
+		if (EmbreeHit.instID[0] != RTC_INVALID_GEOMETRY_ID)
+		{
+			// if instancing is used, material evaluation is deferred here
+			Mesh = Geo.ParentAggregateMesh->StaticMeshInstancesToMappings[EmbreeHit.instID[0]]->Mesh;
+#else
+		Desc = Geo.TriangleDescs[Ray.primID];
 		if (Ray.instID != -1)
 		{
 			// if instancing is used, material evaluation is deferred here
 			Mesh = Geo.ParentAggregateMesh->StaticMeshInstancesToMappings[Ray.instID]->Mesh;
+#endif // USE_EMBREE_MAJOR_VERSION >= 3
+
 			int32 ElementIndex = Desc.ElementIndex;
 			Desc.CastShadow = Mesh->IsElementCastingShadow(ElementIndex);
 			Desc.StaticAndOpaqueMask = !Mesh->IsMasked(ElementIndex) && !Mesh->IsTranslucent(ElementIndex) && !Mesh->bMovable;;
@@ -128,10 +179,21 @@ struct FEmbreeFilterProcessor
 			Desc.Masked = Mesh->IsMasked(ElementIndex);
 			Desc.CastShadowAsMasked = Mesh->IsCastingShadowsAsMasked(ElementIndex);
 		}
+#if USE_EMBREE_MAJOR_VERSION >= 3
+		s = 1 - EmbreeHit.u - EmbreeHit.v;
+#else
 		s = 1 - Ray.u - Ray.v;
+#endif // USE_EMBREE_MAJOR_VERSION >= 3
 	}
 
+#if USE_EMBREE_MAJOR_VERSION >= 3
+	FEmbreeContext& EmbreeContext;
+	RTCRay& EmbreeRay;
+	RTCHit& EmbreeHit;
+	int*& ValidMask;
+#else
 	FEmbreeRay& Ray;
+#endif // USE_EMBREE_MAJOR_VERSION >= 3
 	const FEmbreeGeometry& Geo;
 	FEmbreeTriangleDesc Desc;
 	const FStaticLightingMesh* Mesh;
@@ -152,26 +214,45 @@ struct FEmbreeFilterProcessor
 
 	EMBREE_INLINE void Invalidate()
 	{
+#if USE_EMBREE_MAJOR_VERSION >= 3
+		ValidMask[0] = 0;
+#else
 		Ray.geomID = -1;
+#endif // USE_EMBREE_MAJOR_VERSION >= 3
 	}
 
 	EMBREE_INLINE bool IsBackFace() const
 	{
+#if USE_EMBREE_MAJOR_VERSION >= 3
+		return EmbreeRay.dir_x * -EmbreeHit.Ng_x +  EmbreeRay.dir_y * -EmbreeHit.Ng_y + EmbreeRay.dir_z * -EmbreeHit.Ng_z < 0;
+#else
 		return Ray.dir[0] * Ray.Ng[0] +  Ray.dir[1] * Ray.Ng[1] + Ray.dir[2] * Ray.Ng[2] < 0;
+#endif // USE_EMBREE_MAJOR_VERSION >= 3
 	}
 
 
 	EMBREE_INLINE bool HitRejectedByStaticAndOpaqueOnlyTest() const
 	{
+#if USE_EMBREE_MAJOR_VERSION >= 3
+		return EmbreeContext.bStaticAndOpaqueOnly && !Desc.StaticAndOpaqueMask;
+#else
 		return Ray.bStaticAndOpaqueOnly && !Desc.StaticAndOpaqueMask;
+#endif // USE_EMBREE_MAJOR_VERSION >= 3
 	}
 
 	EMBREE_INLINE bool HitRejectedByBackFaceCullingTest() const
 	{
+#if USE_EMBREE_MAJOR_VERSION >= 3
+		if (!EmbreeContext.bTwoSidedCollision && !Desc.TwoSidedMask)
+		{
+			bool bIsBackFace = IsBackFace();
+			if (EmbreeContext.bFlipSidedness ? !bIsBackFace : bIsBackFace)
+#else
 		if (!Ray.bTwoSidedCollision && !Desc.TwoSidedMask)
 		{
 			bool bIsBackFace = IsBackFace();
 			if (Ray.bFlipSidedness ? !bIsBackFace : bIsBackFace)
+#endif // USE_EMBREE_MAJOR_VERSION >= 3
 			{
 				return true;
 			}
@@ -202,9 +283,15 @@ struct FEmbreeFilterProcessor
      */
 	EMBREE_INLINE bool HitRejectedByHLODTest() const
 	{
+#if USE_EMBREE_MAJOR_VERSION >= 3
+		const FStaticLightingMesh* MappingMesh = EmbreeContext.MappingMesh;
+#else
+		const FStaticLightingMesh* MappingMesh = Ray.MappingMesh;
+#endif // USE_EMBREE_MAJOR_VERSION >= 3
+
 		const uint32 InvalidIndex = 0xFFFF;
 		uint32 GeoHLODTreeIndex = (Mesh->GetLODIndices() & 0xFFFF0000) >> 16;
-		uint32 RayHLODTreeIndex = Ray.MappingMesh ? (Ray.MappingMesh->GetLODIndices() & 0xFFFF0000) >> 16 : InvalidIndex;
+		uint32 RayHLODTreeIndex = MappingMesh ? (MappingMesh->GetLODIndices() & 0xFFFF0000) >> 16 : InvalidIndex;
 
 		bool bRejectHit = false;
 
@@ -215,7 +302,7 @@ struct FEmbreeFilterProcessor
 			uint32 GeoHLODRangeStart = GeoHLODRange & 0xFFFF;
 			uint32 GeoHLODRangeEnd = (GeoHLODRange & 0xFFFF0000) >> 16;
 
-			uint32 RayHLODRange = Ray.MappingMesh ? Ray.MappingMesh->GetHLODRange() : 0;
+			uint32 RayHLODRange = MappingMesh ? MappingMesh->GetHLODRange() : 0;
 			uint32 RayHLODRangeStart = RayHLODRange & 0xFFFF;
 			uint32 RayHLODRangeEnd = (RayHLODRange & 0xFFFF0000) >> 16;
 
@@ -251,18 +338,24 @@ struct FEmbreeFilterProcessor
 	{
 		uint32 GeoMeshLODIndex = Mesh->GetLODIndices() & 0xFFFF;
 
+#if USE_EMBREE_MAJOR_VERSION >= 3
+		const FStaticLightingMesh* MappingMesh = EmbreeContext.MappingMesh;
+#else
+		const FStaticLightingMesh* MappingMesh = Ray.MappingMesh;
+#endif // USE_EMBREE_MAJOR_VERSION >= 3
+
 		// Only shadows from appropriate mesh LODs.
-		if (Ray.MappingMesh)
+		if (MappingMesh)
 		{
 			// If it is not from the same mesh, then only LOD 0 can cast shadow.
-			if (Ray.MappingMesh->MeshIndex != Mesh->MeshIndex)
+			if (MappingMesh->MeshIndex != Mesh->MeshIndex)
 			{
 				return GeoMeshLODIndex != 0;
 			}
 			else
 			{
 				// If it is from the same mesh, then only same LOD can cast shadow.
-				return (Ray.MappingMesh->GetLODIndices() & 0xFFFF) != GeoMeshLODIndex;
+				return (MappingMesh->GetLODIndices() & 0xFFFF) != GeoMeshLODIndex;
 			}
 		}
 
@@ -273,9 +366,15 @@ struct FEmbreeFilterProcessor
 	EMBREE_INLINE bool HitRejectedBySelfShadowTest() const
 	{
 		// No self shadows, or only self shadow
+#if USE_EMBREE_MAJOR_VERSION >= 3
+		if ((Mesh == EmbreeContext.ShadowMesh && ((Mesh->LightingFlags & GI_INSTANCE_SELFSHADOWDISABLE) || (EmbreeContext.TraceFlags & LIGHTRAY_SELFSHADOWDISABLE)))
+		|| (EmbreeContext.bDirectShadowingRay && Desc.IndirectlyShadowedOnly)
+		|| (Mesh != EmbreeContext.ShadowMesh && (Mesh->LightingFlags & GI_INSTANCE_SELFSHADOWONLY)))
+#else
 		if ((Mesh == Ray.ShadowMesh && ((Mesh->LightingFlags & GI_INSTANCE_SELFSHADOWDISABLE) || (Ray.TraceFlags & LIGHTRAY_SELFSHADOWDISABLE)))
 		|| (Ray.bDirectShadowingRay && Desc.IndirectlyShadowedOnly)
 		|| (Mesh != Ray.ShadowMesh && (Mesh->LightingFlags & GI_INSTANCE_SELFSHADOWONLY)))
+#endif // USE_EMBREE_MAJOR_VERSION >= 3
 		{
 			return true;
 		}
@@ -286,7 +385,11 @@ struct FEmbreeFilterProcessor
 	{
 		UpdateCoordinates();
 
+#if USE_EMBREE_MAJOR_VERSION >= 3
+		if (Desc.Masked || (EmbreeContext.bDirectShadowingRay && Desc.CastShadowAsMasked))
+#else
 		if (Desc.Masked || (Ray.bDirectShadowingRay && Desc.CastShadowAsMasked))
+#endif // USE_EMBREE_MAJOR_VERSION >= 3
 		{
  			return !Mesh->EvaluateMaskedCollision(TextureCoordinates, Desc.ElementIndex);
 		}
@@ -298,19 +401,31 @@ void FEmbreeFilterProcessor::UpdateCoordinates()
 {
 	if (bCoordsDirty)
 	{
-		if (Ray.instID == -1)
+#if USE_EMBREE_MAJOR_VERSION >= 3
+		unsigned int InstanceID = EmbreeHit.instID[0];
+		unsigned int PrimID = EmbreeHit.primID;
+		float HitU = EmbreeHit.u;
+		float HitV = EmbreeHit.v;
+#else
+		unsigned int InstanceID = Ray.instID;
+		unsigned int PrimID = Ray.primID;
+		float HitU = Ray.u;
+		float HitV = Ray.v;
+#endif // USE_EMBREE_MAJOR_VERSION >= 3
+
+		if (InstanceID == RTC_INVALID_GEOMETRY_ID)
 		{
-			Mesh->GetTriangleIndices(Ray.primID, Index0, Index1, Index2);
+			Mesh->GetTriangleIndices(PrimID, Index0, Index1, Index2);
 		}
 		else
 		{
-			Mesh->GetInstanceableStaticMesh()->GetNonTransformedTriangleIndices(Ray.primID, Index0, Index1, Index2);
+			Mesh->GetInstanceableStaticMesh()->GetNonTransformedTriangleIndices(PrimID, Index0, Index1, Index2);
 		}
 
 		const FVector2f& UV1 = Geo.UVs[Index0];
 		const FVector2f& UV2 = Geo.UVs[Index1];
 		const FVector2f& UV3 = Geo.UVs[Index2];
-		TextureCoordinates = UV1 * s + UV2 * Ray.u + UV3 * Ray.v;
+		TextureCoordinates = UV1 * s + UV2 * HitU + UV3 * HitV;
 
 		bCoordsDirty = false;
 	}
@@ -319,6 +434,24 @@ void FEmbreeFilterProcessor::UpdateCoordinates()
 // This is called when everything succeeds and the ray is the final collision.
 void FEmbreeFilterProcessor::UpdateRay()
 {
+#if USE_EMBREE_MAJOR_VERSION >= 3
+	// ElementIndex
+	EmbreeContext.ElementIndex = Desc.ElementIndex;
+
+	if (EmbreeContext.bFindClosestIntersection)
+	{
+		//@todo - only lookup and interpolate UV's if needed
+		UpdateCoordinates();
+
+		// TextureCoordinates
+		EmbreeContext.TextureCoordinates = TextureCoordinates;
+
+		// LightmapCoordinates
+		const FVector2f& LightmapUV1 = Geo.LightmapUVs[Index0];
+		const FVector2f& LightmapUV2 = Geo.LightmapUVs[Index1];
+		const FVector2f& LightmapUV3 = Geo.LightmapUVs[Index2];
+		EmbreeContext.LightmapCoordinates = LightmapUV1 * s + LightmapUV2 * EmbreeHit.u + LightmapUV3 * EmbreeHit.v;
+#else
 	// ElementIndex
 	Ray.ElementIndex = Desc.ElementIndex;
 
@@ -335,27 +468,60 @@ void FEmbreeFilterProcessor::UpdateRay()
 		const FVector2f& LightmapUV2 = Geo.LightmapUVs[Index1];
 		const FVector2f& LightmapUV3 = Geo.LightmapUVs[Index2];
 		Ray.LightmapCoordinates = LightmapUV1 * s + LightmapUV2 * Ray.u + LightmapUV3 * Ray.v;
+#endif // USE_EMBREE_MAJOR_VERSION >= 3
 	}
 	else if (bCoordsDirty)
 	{
-		if (Ray.instID == -1)
+#if USE_EMBREE_MAJOR_VERSION >= 3
+		unsigned int InstanceID = EmbreeHit.instID[0];
+		unsigned int PrimID = EmbreeHit.primID;
+#else
+		unsigned int InstanceID = Ray.instID;
+		unsigned int PrimID = Ray.primID;
+#endif // USE_EMBREE_MAJOR_VERSION >= 3
+
+		if (InstanceID == RTC_INVALID_GEOMETRY_ID)
 		{
-			Mesh->GetTriangleIndices(Ray.primID, Index0, Index1, Index2);
+			Mesh->GetTriangleIndices(PrimID, Index0, Index1, Index2);
 		}
 		else
 		{
-			Mesh->GetInstanceableStaticMesh()->GetNonTransformedTriangleIndices(Ray.primID, Index0, Index1, Index2);
+			Mesh->GetInstanceableStaticMesh()->GetNonTransformedTriangleIndices(PrimID, Index0, Index1, Index2);
 		}
 	}
 
 	// Transmission : updated outside of this scope.
 }
 
+#if USE_EMBREE_MAJOR_VERSION >= 3
+// Warning: EmbreeFilterFunc must only modify FEmbreeContext outputs!
+void EmbreeFilterFunc(const RTCFilterFunctionNArguments* args)
+{
+	int* EmbreeValid = args->valid;
+	FEmbreeGeometry& EmbreeGeom = *(FEmbreeGeometry*)args->geometryUserPtr;
+	FEmbreeContext& EmbreeContext = *(FEmbreeContext*)args->context;
+
+	// We expect single-ray packets here since we use rtcIntersect1/rtcOccluded1.
+	check(args->N == 1u);
+
+	// Ignore invalid rays.
+	if (EmbreeValid[0] != -1)
+	{
+		return;
+	}
+
+	RTCRay& EmbreeRay = *(RTCRay*)args->ray;
+	RTCHit& EmbreeHit = *(RTCHit*)args->hit;
+
+	FEmbreeFilterProcessor Proc(EmbreeContext, EmbreeRay, EmbreeHit, EmbreeValid, EmbreeGeom);
+	checkSlow(Proc.Geo.GeomID == Proc.EmbreeHit.geomID);
+#else
 // Warning: EmbreeFilterFunc must only modify RTCRay outputs!  Nothing else is copied back to FEmbreeRay4.
 void EmbreeFilterFunc(void* UserPtr, RTCRay& InRay)
 {
 	FEmbreeFilterProcessor Proc((FEmbreeRay&)InRay, *(FEmbreeGeometry*)UserPtr);
 	checkSlow(Proc.Geo.GeomID == Proc.Ray.geomID);
+#endif // USE_EMBREE_MAJOR_VERSION >= 3
 
 	if (!Proc.Desc.CastShadow)
 	{
@@ -381,6 +547,18 @@ void EmbreeFilterFunc(void* UserPtr, RTCRay& InRay)
 	}
 
 	// No collision with translucent primitives.
+#if USE_EMBREE_MAJOR_VERSION >= 3
+	if (Proc.Desc.Translucent && !(Proc.EmbreeContext.bDirectShadowingRay && Proc.Desc.CastShadowAsMasked))
+	{
+		if (Proc.EmbreeContext.bCalculateTransmission)
+		{
+			Proc.UpdateCoordinates();
+
+			// Accumulate the total transmission along the ray
+			// The result is order independent so the intersections don't have to be strictly front to back
+			Proc.EmbreeContext.TransmissionAcc.Push(Proc.Mesh->EvaluateTransmission(Proc.TextureCoordinates, Proc.Desc.ElementIndex), Proc.EmbreeRay.tfar);
+		}
+#else
 	if (Proc.Desc.Translucent && !(Proc.Ray.bDirectShadowingRay && Proc.Desc.CastShadowAsMasked))
 	{
 		if (Proc.Ray.bCalculateTransmission)
@@ -390,8 +568,9 @@ void EmbreeFilterFunc(void* UserPtr, RTCRay& InRay)
 			// Accumulate the total transmission along the ray
 			// The result is order independent so the intersections don't have to be strictly front to back
 			Proc.Ray.TransmissionAcc.Push(Proc.Mesh->EvaluateTransmission(Proc.TextureCoordinates, Proc.Desc.ElementIndex), Proc.Ray.tfar);
-
 		}
+#endif // USE_EMBREE_MAJOR_VERSION >= 3
+
 		Proc.Invalidate();
 		return;
 	}
@@ -403,7 +582,11 @@ void EmbreeFilterFunc(void* UserPtr, RTCRay& InRay)
 		return;
 	}
 
+#if USE_EMBREE_MAJOR_VERSION >= 3
+	if (Proc.EmbreeContext.bFindClosestIntersection)
+#else
 	if (Proc.Ray.bFindClosestIntersection)
+#endif // USE_EMBREE_MAJOR_VERSION >= 3
 	{
 		if (Proc.HitRejectedByAlphaTestTest())
 		{
@@ -417,6 +600,7 @@ void EmbreeFilterFunc(void* UserPtr, RTCRay& InRay)
 	Proc.UpdateRay();
 }
 
+#if USE_EMBREE_MAJOR_VERSION < 3
 void EmbreeFilterFunc4(const void* valid, void* UserPtr, RTCRay4& InRay)
 {
 	FEmbreeRay4& EmbreeRay4 = (FEmbreeRay4&)InRay;
@@ -433,6 +617,7 @@ void EmbreeFilterFunc4(const void* valid, void* UserPtr, RTCRay4& InRay)
 		}
 	}
 }
+#endif // USE_EMBREE_MAJOR_VERSION < 3
 
 FEmbreeGeometry::FEmbreeGeometry(
 	RTCDevice EmbreeDevice, 
@@ -455,13 +640,25 @@ FEmbreeGeometry::FEmbreeGeometry(
 
 	const FStaticLightingTextureMapping* TextureMapping = Mapping ? Mapping->GetTextureMapping() : NULL;
 
-	GeomID = rtcNewTriangleMesh(EmbreeScene, RTC_GEOMETRY_STATIC, Mesh->NumTriangles, Mesh->NumVertices);
 	TriangleDescs.AddZeroed(Mesh->NumTriangles);
 	UVs.AddZeroed(Mesh->NumVertices);
 	LightmapUVs.AddZeroed(Mesh->NumVertices);
 
+#if USE_EMBREE_MAJOR_VERSION >= 3
+	RTCGeometry EmbreeGeom = rtcNewGeometry(EmbreeDevice, RTC_GEOMETRY_TYPE_TRIANGLE);
+	rtcSetGeometryBuildQuality(EmbreeGeom, RTC_BUILD_QUALITY_MEDIUM);
+	rtcSetGeometryTimeStepCount(EmbreeGeom, 1);
+
+	FVector4f* Vertices = (FVector4f*) rtcSetNewGeometryBuffer(
+		EmbreeGeom, RTC_BUFFER_TYPE_VERTEX, 0, RTC_FORMAT_FLOAT3, 4 * sizeof(float), Mesh->NumVertices);
+	int32* Indices = (int32*) rtcSetNewGeometryBuffer(
+		EmbreeGeom, RTC_BUFFER_TYPE_INDEX, 0, RTC_FORMAT_UINT3, 3 * sizeof(int32), Mesh->NumTriangles);
+#else
+	GeomID = rtcNewTriangleMesh(EmbreeScene, RTC_GEOMETRY_STATIC, Mesh->NumTriangles, Mesh->NumVertices);
+
 	FVector4f* Vertices = (FVector4f*) rtcMapBuffer(EmbreeScene, GeomID, RTC_VERTEX_BUFFER);
 	int32* Indices = (int32*) rtcMapBuffer(EmbreeScene, GeomID, RTC_INDEX_BUFFER);
+#endif // USE_EMBREE_MAJOR_VERSION >= 3
 
 	for (int32 TriangleIndex = 0;TriangleIndex < Mesh->NumTriangles;TriangleIndex++)
 	{
@@ -547,10 +744,19 @@ FEmbreeGeometry::FEmbreeGeometry(
 		}
 
 	}
+
+#if USE_EMBREE_MAJOR_VERSION >= 3
+	rtcCommitGeometry(EmbreeGeom);
+	GeomID = rtcAttachGeometry(EmbreeScene, EmbreeGeom);
+	rtcReleaseGeometry(EmbreeGeom);
+
+	check(rtcGetDeviceError(EmbreeDevice) == RTC_ERROR_NONE);
+#else
 	rtcUnmapBuffer(EmbreeScene, GeomID, RTC_VERTEX_BUFFER);
 	rtcUnmapBuffer(EmbreeScene, GeomID, RTC_INDEX_BUFFER);
 
 	check(rtcDeviceGetError(EmbreeDevice) == RTC_NO_ERROR);
+#endif // USE_EMBREE_MAJOR_VERSION >= 3
 }
 
 void CalculateSurfaceArea(const FStaticLightingMesh* Mesh, const FBoxSphereBounds3f& ImportanceBounds, float& SurfaceArea, float& SurfaceAreaWithinImportanceVolume)
@@ -586,7 +792,11 @@ void CalculateSurfaceArea(const FStaticLightingMesh* Mesh, const FBoxSphereBound
 	}
 }
 
+#if USE_EMBREE_MAJOR_VERSION >= 3
+bool EmbreeMemoryMonitor(void* ptr, ssize_t bytes, bool post)
+#else
 bool EmbreeMemoryMonitor(const ssize_t bytes, const bool post)
+#endif // USE_EMBREE_MAJOR_VERSION >= 3
 {
 	FPlatformAtomics::InterlockedAdd(&GEmbreeAllocatedSpace, bytes);
 	return true;
@@ -599,9 +809,17 @@ FEmbreeAggregateMesh::FEmbreeAggregateMesh(const FScene& InScene):
 	TotalNumTriangles(0),
 	TotalNumTrianglesInstanced(0)
 {
-	rtcDeviceSetMemoryMonitorFunction(InScene.EmbreeDevice, EmbreeMemoryMonitor);
-
 	EmbreeDevice = InScene.EmbreeDevice;
+
+#if USE_EMBREE_MAJOR_VERSION >= 3
+	rtcSetDeviceMemoryMonitorFunction(EmbreeDevice, EmbreeMemoryMonitor, NULL);
+
+	EmbreeScene = rtcNewScene(EmbreeDevice);
+	rtcSetSceneBuildQuality(EmbreeScene, RTC_BUILD_QUALITY_MEDIUM);
+	check(rtcGetDeviceError(EmbreeDevice) == RTC_ERROR_NONE);
+#else
+	rtcDeviceSetMemoryMonitorFunction(EmbreeDevice, EmbreeMemoryMonitor);
+
 	uint32 AlgorithmFlags = RTC_INTERSECT1;
 
 	if (InScene.GeneralSettings.bUseEmbreePacketTracing)
@@ -609,15 +827,21 @@ FEmbreeAggregateMesh::FEmbreeAggregateMesh(const FScene& InScene):
 		AlgorithmFlags |= RTC_INTERSECT4;
 	}
 
-	EmbreeScene = rtcDeviceNewScene(InScene.EmbreeDevice, RTC_SCENE_STATIC, (RTCAlgorithmFlags)AlgorithmFlags);
+	EmbreeScene = rtcDeviceNewScene(EmbreeDevice, RTC_SCENE_STATIC, (RTCAlgorithmFlags)AlgorithmFlags);
 	check(rtcDeviceGetError(EmbreeDevice) == RTC_NO_ERROR);
+#endif // USE_EMBREE_MAJOR_VERSION >= 3
 
 	if (InScene.GeneralSettings.bUseEmbreePacketTracing)
 	{
+#if USE_EMBREE_MAJOR_VERSION >= 3
+		ssize_t SupportsPacketTracing = rtcGetDeviceProperty(EmbreeDevice, RTC_DEVICE_PROPERTY_NATIVE_RAY4_SUPPORTED);
+		check(SupportsPacketTracing);
+#else
 #if RTCORE_VERSION_MAJOR >= 2 && RTCORE_VERSION_MINOR >= 14
 		ssize_t SupportsPacketTracing = rtcDeviceGetParameter1i(EmbreeDevice, RTC_CONFIG_INTERSECT4);
 		check(SupportsPacketTracing);
 #endif
+#endif // USE_EMBREE_MAJOR_VERSION >= 3
 	}
 }
 
@@ -627,7 +851,12 @@ FEmbreeAggregateMesh::~FEmbreeAggregateMesh()
 	{
 		delete MeshInfos[i];
 	}
+
+#if USE_EMBREE_MAJOR_VERSION >= 3
+	rtcReleaseScene(EmbreeScene);
+#else
 	rtcDeleteScene(EmbreeScene);
+#endif // USE_EMBREE_MAJOR_VERSION >= 3
 }
 
 /**
@@ -649,13 +878,26 @@ void FEmbreeAggregateMesh::AddMesh(const FStaticLightingMesh* Mesh, const FStati
 
 			if (!StaticMeshGeometries.Find(LOD))
 			{
+#if USE_EMBREE_MAJOR_VERSION >= 3
+				RTCScene MeshScene = rtcNewScene(EmbreeDevice);
+				rtcSetSceneBuildQuality(MeshScene, RTC_BUILD_QUALITY_MEDIUM);
+#else
 				RTCScene MeshScene = rtcDeviceNewScene(EmbreeDevice, RTC_SCENE_STATIC, RTC_INTERSECT1);
+#endif // USE_EMBREE_MAJOR_VERSION >= 3
 
 				FEmbreeGeometry* Geo = new FEmbreeGeometry(EmbreeDevice, MeshScene, Scene.GetImportanceBounds(), Mesh, Mapping, true);
 				Geo->ParentAggregateMesh = this;
 				StaticMeshGeometries.Add(LOD, FEmbreeStaticMeshGeometry { MeshScene, Geo });
 				MeshInfos.Add(Geo);
 
+#if USE_EMBREE_MAJOR_VERSION >= 3
+				rtcSetGeometryIntersectFilterFunction(rtcGetGeometry(MeshScene, Geo->GeomID), EmbreeFilterFunc);
+				rtcSetGeometryOccludedFilterFunction(rtcGetGeometry(MeshScene, Geo->GeomID), EmbreeFilterFunc);
+
+				rtcSetGeometryUserData(rtcGetGeometry(MeshScene, Geo->GeomID), Geo);
+
+				rtcCommitScene(MeshScene);
+#else
 				rtcSetIntersectionFilterFunction(MeshScene, Geo->GeomID, EmbreeFilterFunc);
 				rtcSetOcclusionFilterFunction(MeshScene, Geo->GeomID, EmbreeFilterFunc);
 
@@ -665,6 +907,7 @@ void FEmbreeAggregateMesh::AddMesh(const FStaticLightingMesh* Mesh, const FStati
 				rtcSetUserData(MeshScene, Geo->GeomID, Geo);
 
 				rtcCommit(MeshScene);
+#endif // USE_EMBREE_MAJOR_VERSION >= 3
 
 				bHasShadowCastingPrimitives |= Geo->bHasShadowCastingPrimitives;
 
@@ -681,9 +924,22 @@ void FEmbreeAggregateMesh::AddMesh(const FStaticLightingMesh* Mesh, const FStati
 			SceneSurfaceArea += SurfaceArea;
 			SceneSurfaceAreaWithinImportanceVolume += SurfaceAreaWithinImportanceVolume;
 
+#if USE_EMBREE_MAJOR_VERSION >= 3
+			RTCGeometry EmbreeGeom = rtcNewGeometry(EmbreeDevice, RTC_GEOMETRY_TYPE_INSTANCE);
+			rtcSetGeometryInstancedScene(EmbreeGeom, StaticMeshGeometries[LOD].MeshScene);
+			rtcSetGeometryTimeStepCount(EmbreeGeom, 1);
+			rtcCommitGeometry(EmbreeGeom);
+			unsigned int InstID = rtcAttachGeometry(EmbreeScene, EmbreeGeom);
+			rtcSetGeometryUserData(EmbreeGeom, (void*)Mapping);
+			rtcSetGeometryTransform(EmbreeGeom, 0, RTC_FORMAT_FLOAT4X4_COLUMN_MAJOR, (const float*)StaticMeshInstance->LocalToWorld.M);
+			rtcReleaseGeometry(EmbreeGeom);
+			check(rtcGetDeviceError(EmbreeDevice) == RTC_ERROR_NONE);
+#else
 			auto InstID = rtcNewInstance(EmbreeScene, StaticMeshGeometries[LOD].MeshScene);
 			rtcSetUserData(EmbreeScene, InstID, (void*)Mapping);
 			rtcSetTransform(EmbreeScene, InstID, RTC_MATRIX_COLUMN_MAJOR_ALIGNED16, (const float*)StaticMeshInstance->LocalToWorld.M);
+#endif // USE_EMBREE_MAJOR_VERSION >= 3
+
 			if (StaticMeshInstancesToMappings.Num() < (int32)InstID + 1)
 			{
 				StaticMeshInstancesToMappings.SetNum(InstID + 1);
@@ -696,12 +952,18 @@ void FEmbreeAggregateMesh::AddMesh(const FStaticLightingMesh* Mesh, const FStati
 			FEmbreeGeometry* Geo = new FEmbreeGeometry(EmbreeDevice, EmbreeScene, Scene.GetImportanceBounds(), Mesh, Mapping, false);
 			MeshInfos.Add(Geo);
 
+#if USE_EMBREE_MAJOR_VERSION >= 3
+			rtcSetGeometryUserData(rtcGetGeometry(EmbreeScene, Geo->GeomID), Geo);
+			rtcSetGeometryIntersectFilterFunction(rtcGetGeometry(EmbreeScene, Geo->GeomID), EmbreeFilterFunc);
+			rtcSetGeometryOccludedFilterFunction(rtcGetGeometry(EmbreeScene, Geo->GeomID), EmbreeFilterFunc);
+#else
 			rtcSetUserData(EmbreeScene, Geo->GeomID, Geo);
 			rtcSetIntersectionFilterFunction(EmbreeScene, Geo->GeomID, EmbreeFilterFunc);
 			rtcSetOcclusionFilterFunction(EmbreeScene, Geo->GeomID, EmbreeFilterFunc);
 
 			rtcSetIntersectionFilterFunction4(EmbreeScene, Geo->GeomID, EmbreeFilterFunc4);
 			rtcSetOcclusionFilterFunction4(EmbreeScene, Geo->GeomID, EmbreeFilterFunc4);
+#endif // USE_EMBREE_MAJOR_VERSION >= 3
 
 			bHasShadowCastingPrimitives |= Geo->bHasShadowCastingPrimitives;
 
@@ -717,8 +979,13 @@ void FEmbreeAggregateMesh::PrepareForRaytracing()
 {
 	const double StartTime = FPlatformTime::Seconds();
 
+#if USE_EMBREE_MAJOR_VERSION >= 3
+	rtcCommitScene(EmbreeScene);
+	check(rtcGetDeviceError(EmbreeDevice) == RTC_ERROR_NONE);
+#else
 	rtcCommit(EmbreeScene);
 	check(rtcDeviceGetError(EmbreeDevice) == RTC_NO_ERROR);
+#endif // USE_EMBREE_MAJOR_VERSION >= 3
 
 	const float Buildtime = FPlatformTime::Seconds() - StartTime;
 	UE_LOG(LogLightmass, Log, TEXT("Embree Build %.1fs"), Buildtime);
@@ -770,6 +1037,39 @@ bool FEmbreeAggregateMesh::IntersectLightRay(
 
 	ClosestIntersection.bIntersects = false;
 
+#if USE_EMBREE_MAJOR_VERSION >= 3
+	FEmbreeContext EmbreeContext(
+		LightRay.Mesh,
+		LightRay.Mapping ? LightRay.Mapping->Mesh : NULL,
+		LightRay.TraceFlags,
+		bFindClosestIntersection,
+		bCalculateTransmission,
+		bDirectShadowingRay);
+#if USE_EMBREE_MAJOR_VERSION >= 4
+	rtcInitRayQueryContext(&EmbreeContext);
+#else
+	rtcInitIntersectContext(&EmbreeContext);
+#endif USE_EMBREE_MAJOR_VERSION >= 4
+
+	RTCRayHit EmbreeRayHit;
+	EmbreeRayHit.ray.org_x = LightRay.Start.X;
+	EmbreeRayHit.ray.org_y = LightRay.Start.Y;
+	EmbreeRayHit.ray.org_z = LightRay.Start.Z;
+	EmbreeRayHit.ray.tnear = 0.0f;
+
+	EmbreeRayHit.ray.dir_x = LightRay.Direction.X;
+	EmbreeRayHit.ray.dir_y = LightRay.Direction.Y;
+	EmbreeRayHit.ray.dir_z = LightRay.Direction.Z;
+	EmbreeRayHit.ray.tfar = LightRay.Length;
+
+	EmbreeRayHit.ray.time = 0.0f;
+	EmbreeRayHit.ray.mask = 0xFFFFFFFF;
+	EmbreeRayHit.ray.flags = 0u;
+
+	EmbreeRayHit.hit.u = EmbreeRayHit.hit.v = 0.0f;
+
+	EmbreeRayHit.hit.geomID = RTC_INVALID_GEOMETRY_ID;
+#else
 	FEmbreeRay EmbreeRay(
 		LightRay.Mesh, 
 		LightRay.Mapping ? LightRay.Mapping->Mesh : NULL,
@@ -786,16 +1086,69 @@ bool FEmbreeAggregateMesh::IntersectLightRay(
 	EmbreeRay.dir[2] = LightRay.Direction.Z;
 	EmbreeRay.tnear = 0;
 	EmbreeRay.tfar = LightRay.Length;
+#endif // USE_EMBREE_MAJOR_VERSION >= 3
 
 	if (bFindClosestIntersection)
 	{
+#if USE_EMBREE_MAJOR_VERSION >= 3
+#if USE_EMBREE_MAJOR_VERSION >= 4
+		RTCIntersectArguments args;
+		rtcInitIntersectArguments(&args);
+		args.context = &EmbreeContext;
+		rtcIntersect1(EmbreeScene, &EmbreeRayHit, &args);
+#else
+		rtcIntersect1(EmbreeScene, &EmbreeContext, &EmbreeRayHit);
+#endif // USE_EMBREE_MAJOR_VERSION >= 4
+		EmbreeRayHit.hit.Ng_x = -EmbreeRayHit.hit.Ng_x;
+		EmbreeRayHit.hit.Ng_y = -EmbreeRayHit.hit.Ng_y;
+		EmbreeRayHit.hit.Ng_z = -EmbreeRayHit.hit.Ng_z;
+#else
 		rtcIntersect(EmbreeScene, EmbreeRay);
+#endif // USE_EMBREE_MAJOR_VERSION >= 3
 	}
 	else
 	{
+#if USE_EMBREE_MAJOR_VERSION >= 4
+		RTCOccludedArguments args;
+		rtcInitOccludedArguments(&args);
+		args.context = &EmbreeContext;
+		rtcOccluded1(EmbreeScene, &EmbreeRayHit.ray, &args);
+#elif USE_EMBREE_MAJOR_VERSION == 3
+		rtcOccluded1(EmbreeScene, &EmbreeContext, &EmbreeRayHit.ray);
+#else
 		rtcOccluded(EmbreeScene, EmbreeRay);
+#endif // USE_EMBREE_MAJOR_VERSION >= 4
 	}
 
+#if USE_EMBREE_MAJOR_VERSION >= 3
+	if (EmbreeRayHit.ray.tfar >= 0.0f && EmbreeRayHit.hit.geomID != -1 && EmbreeRayHit.hit.primID != -1)
+	{
+		FMinimalStaticLightingVertex EmbreeVertex;
+		EmbreeVertex.WorldPosition = LightRay.Start + LightRay.Direction * EmbreeRayHit.ray.tfar;
+
+		EmbreeVertex.TextureCoordinates[0] = EmbreeContext.TextureCoordinates;
+		EmbreeVertex.TextureCoordinates[1] = EmbreeContext.LightmapCoordinates;
+
+		if (EmbreeRayHit.hit.instID[0] == -1)
+		{
+			const FEmbreeGeometry& Geo = *(FEmbreeGeometry*)rtcGetGeometryUserData(rtcGetGeometry(EmbreeScene, EmbreeRayHit.hit.geomID));
+			EmbreeVertex.WorldTangentZ = FVector3f(EmbreeRayHit.hit.Ng_x, EmbreeRayHit.hit.Ng_y, EmbreeRayHit.hit.Ng_z).GetSafeNormal();
+			ClosestIntersection = FLightRayIntersection(true, EmbreeVertex, Geo.Mesh, Geo.Mapping, EmbreeContext.ElementIndex);
+		}
+		else
+		{
+			FStaticLightingMapping* Mapping = (FStaticLightingMapping*)rtcGetGeometryUserData(rtcGetGeometry(EmbreeScene, EmbreeRayHit.hit.instID[0]));
+			FVector3f GeometryNormal(EmbreeRayHit.hit.Ng_x, EmbreeRayHit.hit.Ng_y, EmbreeRayHit.hit.Ng_z);
+			EmbreeVertex.WorldTangentZ = Mapping->Mesh->GetInstanceableStaticMesh()->LocalToWorldInverseTranspose.TransformVector(GeometryNormal).GetSafeNormal();
+			ClosestIntersection = FLightRayIntersection(true, EmbreeVertex, Mapping->Mesh, Mapping, EmbreeContext.ElementIndex);
+		}
+		EmbreeContext.TransmissionAcc.Resolve(ClosestIntersection.Transmission, EmbreeRayHit.ray.tfar);
+	}
+	else
+	{
+		EmbreeContext.TransmissionAcc.Resolve(ClosestIntersection.Transmission);
+	}
+#else
 	if (EmbreeRay.geomID != -1 && EmbreeRay.primID != -1)
 	{
 		FMinimalStaticLightingVertex EmbreeVertex;
@@ -823,6 +1176,7 @@ bool FEmbreeAggregateMesh::IntersectLightRay(
 	{
 		EmbreeRay.TransmissionAcc.Resolve(ClosestIntersection.Transmission);
 	}
+#endif // USE_EMBREE_MAJOR_VERSION >= 3
 
 	return ClosestIntersection.bIntersects;
 }
diff --git a/Engine/Source/Programs/UnrealLightmass/Private/Lighting/Embree.h b/Engine/Source/Programs/UnrealLightmass/Private/Lighting/Embree.h
--- a/Engine/Source/Programs/UnrealLightmass/Private/Lighting/Embree.h
+++ b/Engine/Source/Programs/UnrealLightmass/Private/Lighting/Embree.h
@@ -11,8 +11,12 @@ namespace Lightmass
 {	
 	class FStaticLightingMesh;
 
+#if USE_EMBREE_MAJOR_VERSION >= 3
+	void EmbreeFilterFunc(const RTCFilterFunctionNArguments* args);
+#else
 	void EmbreeFilterFunc(void* UserPtr, RTCRay& InRay);
 	void EmbreeFilterFunc4(const void* valid, void* UserPtr, RTCRay4& InRay);
+#endif // USE_EMBREE_MAJOR_VERSION >= 3
 
 	// Accumulates transmission color, and store the distance.
 	// This is required since Embree test collision in any order, 
@@ -29,6 +33,40 @@ namespace Lightmass
 		void Resolve(FLinearColor& FinalColor);
 	};
 
+#if USE_EMBREE_MAJOR_VERSION >= 3
+#if USE_EMBREE_MAJOR_VERSION >= 4
+	struct FEmbreeContext : public RTCRayQueryContext
+#else
+	struct FEmbreeContext : public RTCIntersectContext
+#endif // USE_EMBREE_MAJOR_VERSION >= 4
+	{
+		FEmbreeContext(
+			const FStaticLightingMesh* InShadowMesh,
+			const FStaticLightingMesh* InMappingMesh,
+			uint32 InTraceFlags,
+			bool InFindClosestIntersection,
+			bool InCalculateTransmission,
+			bool InDirectShadowingRay
+			);
+
+		const FStaticLightingMesh* ShadowMesh; // Controls self shadow behavior, such as only self shadow and no self shadow.
+		const FStaticLightingMesh* MappingMesh; // Controls LOD behavior to control which mesh are allowed to cast shadows
+		uint32 TraceFlags;
+		bool bFindClosestIntersection;
+		bool bCalculateTransmission;
+		bool bDirectShadowingRay;
+		bool bStaticAndOpaqueOnly;
+		bool bTwoSidedCollision;
+		bool bFlipSidedness;
+
+		// Outputs
+		// Warning: EmbreeFilterFunc must only modify these!
+		int32 ElementIndex; // Material Index
+		FVector2f TextureCoordinates; // Material Coordinates
+		FVector2f LightmapCoordinates;
+		FEmbreeTransmissionAccumulator TransmissionAcc;
+	};
+#else
 	struct FEmbreeRay : public RTCRay
 	{
 		FEmbreeRay(
@@ -142,6 +180,7 @@ namespace Lightmass
 		FVector2f LightmapCoordinates[4];
 		FEmbreeTransmissionAccumulator TransmissionAcc[4];
 	};
+#endif // USE_EMBREE_MAJOR_VERSION >= 3
 
 	struct FEmbreeTriangleDesc
 	{
diff --git a/Engine/Source/Programs/UnrealLightmass/Private/Lighting/Lighting.h b/Engine/Source/Programs/UnrealLightmass/Private/Lighting/Lighting.h
--- a/Engine/Source/Programs/UnrealLightmass/Private/Lighting/Lighting.h
+++ b/Engine/Source/Programs/UnrealLightmass/Private/Lighting/Lighting.h
@@ -13,8 +13,16 @@
 #define ALLOW_LIGHTMAP_SAMPLE_DEBUGGING	0
 
 #if USE_EMBREE
-#include <embree2/rtcore.h>
-#include <embree2/rtcore_ray.h>
+#if USE_EMBREE_MAJOR_VERSION >= 4
+	#include <embree4/rtcore.h>
+	#include <embree4/rtcore_ray.h>
+#elif USE_EMBREE_MAJOR_VERSION == 3
+	#include <embree3/rtcore.h>
+	#include <embree3/rtcore_ray.h>
+#else
+	#include <embree2/rtcore.h>
+	#include <embree2/rtcore_ray.h>
+#endif // USE_EMBREE_MAJOR_VERSION >= 4
 #else
 typedef void* RTCDevice;
 #endif
diff --git a/Engine/Source/Programs/UnrealLightmass/UnrealLightmass.Build.cs b/Engine/Source/Programs/UnrealLightmass/UnrealLightmass.Build.cs
--- a/Engine/Source/Programs/UnrealLightmass/UnrealLightmass.Build.cs
+++ b/Engine/Source/Programs/UnrealLightmass/UnrealLightmass.Build.cs
@@ -71,6 +71,7 @@ public class UnrealLightmass : ModuleRules
             RuntimeDependencies.Add("$(EngineDir)/Binaries/Win64/tbb.dll", IntelTBBLibs + "Win64/tbb.dll");
             RuntimeDependencies.Add("$(EngineDir)/Binaries/Win64/tbbmalloc.dll", IntelTBBLibs + "Win64/tbbmalloc.dll");
             PublicDefinitions.Add("USE_EMBREE=1");
+			PublicDefinitions.Add("USE_EMBREE_MAJOR_VERSION=2");
         }
 		else if (Target.Platform == UnrealTargetPlatform.Mac)
 		{
@@ -82,6 +83,7 @@ public class UnrealLightmass : ModuleRules
 			RuntimeDependencies.Add("$(EngineDir)/Binaries/Mac/libtbb.dylib", IntelTBBLibs + "Mac/libtbb.dylib"); // Take latest version to avoid overwriting the editor's copy
 			RuntimeDependencies.Add("$(EngineDir)/Binaries/Mac/libtbbmalloc.dylib", IntelTBBLibs + "Mac/libtbbmalloc.dylib");
 			PublicDefinitions.Add("USE_EMBREE=1");
+			PublicDefinitions.Add("USE_EMBREE_MAJOR_VERSION=2");
 		}
 		else if (Target.IsInPlatformGroup(UnrealPlatformGroup.Unix) && Target.Architecture == UnrealArch.X64) // no support for arm64 yet
 		{
@@ -100,6 +102,7 @@ public class UnrealLightmass : ModuleRules
 			// disabled for Linux atm due to a bug in libtbbmalloc on exit
 			// RuntimeDependencies.Add(Path.Combine(IntelTBBLibs, "Linux/libtbbmalloc.so"));
 			PublicDefinitions.Add("USE_EMBREE=1");
+			PublicDefinitions.Add("USE_EMBREE_MAJOR_VERSION=2");
 		}
         else
         {
