UnrealLightmass: use the engine's version of Embree

Given the version of Embree currently used by the engine, this
represents an upgrade of the version of Embree used by
UnrealLightmass from 2.7.0 to 3.12.2.

This is work towards upgrading the engine version of
Embree to 4.3.3.

diff --git a/Engine/Source/Editor/UnrealEd/Private/Lightmass/Lightmass.cpp b/Engine/Source/Editor/UnrealEd/Private/Lightmass/Lightmass.cpp
--- a/Engine/Source/Editor/UnrealEd/Private/Lightmass/Lightmass.cpp
+++ b/Engine/Source/Editor/UnrealEd/Private/Lightmass/Lightmass.cpp
@@ -3038,8 +3038,8 @@ bool FLightmassProcessor::BeginRun()
 		TEXT("Binaries/Win64/UnrealLightmass-Json.dll"),
 		TEXT("Binaries/Win64/UnrealLightmass-SSL.dll"),
 		TEXT("Binaries/Win64/UnrealLightmass-BuildSettings.dll"),
-		TEXT("Binaries/Win64/embree.dll"),
-		TEXT("Binaries/Win64/tbb.dll"),
+		TEXT("Binaries/Win64/embree3.dll"),
+		TEXT("Binaries/Win64/tbb12.dll"),
 		TEXT("Binaries/Win64/tbbmalloc.dll")
 	};
 #elif PLATFORM_MAC
@@ -3060,8 +3060,8 @@ bool FLightmassProcessor::BeginRun()
 		TEXT("Binaries/Mac/UnrealLightmass-SSL.dylib"),
 		TEXT("Binaries/Mac/UnrealLightmass-SwarmInterface.dylib"),
 		TEXT("Binaries/Mac/UnrealLightmass-BuildSettings.dylib"),
-		TEXT("Binaries/Mac/libembree.2.dylib"),
-		TEXT("Binaries/Mac/libtbb.dylib"),
+		TEXT("Binaries/Mac/libembree3.3.dylib"),
+		TEXT("Binaries/Mac/libtbb.12.dylib"),
 		TEXT("Binaries/Mac/libtbbmalloc.dylib")
 	};
 #elif PLATFORM_LINUX
diff --git a/Engine/Source/Programs/UnrealLightmass/Private/Lighting/Embree.cpp b/Engine/Source/Programs/UnrealLightmass/Private/Lighting/Embree.cpp
--- a/Engine/Source/Programs/UnrealLightmass/Private/Lighting/Embree.cpp
+++ b/Engine/Source/Programs/UnrealLightmass/Private/Lighting/Embree.cpp
@@ -1049,7 +1049,7 @@ bool FEmbreeAggregateMesh::IntersectLightRay(
 	rtcInitRayQueryContext(&EmbreeContext);
 #else
 	rtcInitIntersectContext(&EmbreeContext);
-#endif USE_EMBREE_MAJOR_VERSION >= 4
+#endif // USE_EMBREE_MAJOR_VERSION >= 4
 
 	RTCRayHit EmbreeRayHit;
 	EmbreeRayHit.ray.org_x = LightRay.Start.X;
diff --git a/Engine/Source/Programs/UnrealLightmass/UnrealLightmass.Build.cs b/Engine/Source/Programs/UnrealLightmass/UnrealLightmass.Build.cs
--- a/Engine/Source/Programs/UnrealLightmass/UnrealLightmass.Build.cs
+++ b/Engine/Source/Programs/UnrealLightmass/UnrealLightmass.Build.cs
@@ -15,6 +15,9 @@ public class UnrealLightmass : ModuleRules
 
 		if (Target.Platform == UnrealTargetPlatform.Win64)
 		{
+			PublicDependencyModuleNames.Add("IntelTBB");
+			PublicDependencyModuleNames.Add("Embree");
+
 			// Unreallightmass requires GetProcessMemoryInfo exported by psapi.dll. http://msdn.microsoft.com/en-us/library/windows/desktop/ms683219(v=vs.85).aspx
 			PublicSystemLibraries.Add("psapi.lib");
 			PrivateDependencyModuleNames.AddRange(
@@ -25,6 +28,9 @@ public class UnrealLightmass : ModuleRules
 		}
 		else if (Target.Platform == UnrealTargetPlatform.Mac || Target.Platform == UnrealTargetPlatform.Linux)
 		{
+			PublicDependencyModuleNames.Add("IntelTBB");
+			PublicDependencyModuleNames.Add("Embree");
+
 			// On Mac/Linux UnrealLightmass is executed locally and communicates with the editor using Messaging module instead of SwarmAgent
 			// @todo: allow for better plug-in support in standalone Slate apps
 			PrivateDependencyModuleNames.AddRange(
@@ -42,6 +48,10 @@ public class UnrealLightmass : ModuleRules
 				}
 			);
 		}
+		else
+		{
+			PublicDefinitions.Add("USE_EMBREE=0");
+		}
 		
 		// Lightmass ray tracing is 8% faster with buffer security checks disabled due to fixed size arrays on the stack in the kDOP ray tracing functions
 		// Warning: This means buffer overwrites will not be detected
@@ -55,58 +65,5 @@ public class UnrealLightmass : ModuleRules
 		PrivateIncludePaths.Add("Programs/UnrealLightmass/Private/LightmassCore/Misc");
 		PrivateIncludePaths.Add("Programs/UnrealLightmass/Private/LightmassCore/Math");
 		PrivateIncludePaths.Add("Programs/UnrealLightmass/Private/LightmassCore/Templates");
-
-        // Always use the official version of IntelTBB
-        string IntelTBBLibs = Target.UEThirdPartyBinariesDirectory + "Intel/TBB/";
-        string IntelEmbreeLibs = Target.UEThirdPartyBinariesDirectory + "Intel/Embree/Embree270";
-
-        // EMBREE
-        if (Target.Platform == UnrealTargetPlatform.Win64)
-        {
-            string SDKDir = Target.UEThirdPartySourceDirectory + "Intel/Embree/Embree270/Win64/";
-
-            PublicIncludePaths.Add(SDKDir + "include");
-            PublicAdditionalLibraries.Add(SDKDir + "lib/embree.lib");
-            RuntimeDependencies.Add("$(EngineDir)/Binaries/Win64/embree.dll");
-            RuntimeDependencies.Add("$(EngineDir)/Binaries/Win64/tbb.dll", IntelTBBLibs + "Win64/tbb.dll");
-            RuntimeDependencies.Add("$(EngineDir)/Binaries/Win64/tbbmalloc.dll", IntelTBBLibs + "Win64/tbbmalloc.dll");
-            PublicDefinitions.Add("USE_EMBREE=1");
-			PublicDefinitions.Add("USE_EMBREE_MAJOR_VERSION=2");
-        }
-		else if (Target.Platform == UnrealTargetPlatform.Mac)
-		{
-            string SDKDir = Target.UEThirdPartySourceDirectory + "Intel/Embree/Embree270/MacOSX/";
-
-            PublicIncludePaths.Add(SDKDir + "include");
-            PublicAdditionalLibraries.Add(SDKDir + "lib/libembree.2.dylib");
-			RuntimeDependencies.Add("$(EngineDir)/Binaries/Mac/libembree.2.dylib");
-			RuntimeDependencies.Add("$(EngineDir)/Binaries/Mac/libtbb.dylib", IntelTBBLibs + "Mac/libtbb.dylib"); // Take latest version to avoid overwriting the editor's copy
-			RuntimeDependencies.Add("$(EngineDir)/Binaries/Mac/libtbbmalloc.dylib", IntelTBBLibs + "Mac/libtbbmalloc.dylib");
-			PublicDefinitions.Add("USE_EMBREE=1");
-			PublicDefinitions.Add("USE_EMBREE_MAJOR_VERSION=2");
-		}
-		else if (Target.IsInPlatformGroup(UnrealPlatformGroup.Unix) && Target.Architecture == UnrealArch.X64) // no support for arm64 yet
-		{
-			string IncludeDir = Target.UEThirdPartySourceDirectory + "Intel/Embree/Embree270/Linux/x86_64-unknown-linux-gnu";
-			string SDKDir = Path.Combine(IntelEmbreeLibs, "Linux/x86_64-unknown-linux-gnu/lib");
-
-			PublicIncludePaths.Add(Path.Combine(IncludeDir, "include"));
-			PublicAdditionalLibraries.Add(Path.Combine(SDKDir, "libembree.so"));
-			PublicAdditionalLibraries.Add(Path.Combine(IntelTBBLibs, "Linux/libtbb.so"));
-			// disabled for Linux atm due to a bug in libtbbmalloc on exit
-			// PublicAdditionalLibraries.Add(Path.Combine(IntelTBBLibs, "Linux/libtbbmalloc.so"));
-			RuntimeDependencies.Add(Path.Combine(SDKDir, "libembree.so"));
-			RuntimeDependencies.Add(Path.Combine(SDKDir, "libembree.so.2"));
-			RuntimeDependencies.Add(Path.Combine(IntelTBBLibs, "Linux/libtbb.so"));
-			RuntimeDependencies.Add(Path.Combine(IntelTBBLibs, "Linux/libtbb.so.2"));
-			// disabled for Linux atm due to a bug in libtbbmalloc on exit
-			// RuntimeDependencies.Add(Path.Combine(IntelTBBLibs, "Linux/libtbbmalloc.so"));
-			PublicDefinitions.Add("USE_EMBREE=1");
-			PublicDefinitions.Add("USE_EMBREE_MAJOR_VERSION=2");
-		}
-        else
-        {
-            PublicDefinitions.Add("USE_EMBREE=0");
-        }
     }
 }
