[UnrealLightMass]
* Fixed mac compile error due to missing include

diff --git a/Engine/Source/Programs/UnrealLightmass/Private/Lighting/LightingSystem.cpp b/Engine/Source/Programs/UnrealLightmass/Private/Lighting/LightingSystem.cpp
--- a/Engine/Source/Programs/UnrealLightmass/Private/Lighting/LightingSystem.cpp
+++ b/Engine/Source/Programs/UnrealLightmass/Private/Lighting/LightingSystem.cpp
@@ -16,6 +16,10 @@
 #include "Async/TaskGraphInterfaces.h"
 #endif
 
+#if PLATFORM_MAC
+#include "Mac/MacSystemIncludes.h"
+#endif
+
 #if PLATFORM_WINDOWS
 #include "Windows/AllowWindowsPlatformTypes.h"
 	#include <psapi.h>
diff --git a/Engine/Source/Programs/UnrealLightmass/UnrealLightmass.Build.cs b/Engine/Source/Programs/UnrealLightmass/UnrealLightmass.Build.cs
--- a/Engine/Source/Programs/UnrealLightmass/UnrealLightmass.Build.cs
+++ b/Engine/Source/Programs/UnrealLightmass/UnrealLightmass.Build.cs
@@ -7,6 +7,8 @@ public class UnrealLightmass : ModuleRules
 {
 	public UnrealLightmass(ReadOnlyTargetRules Target) : base(Target)
 	{
+		bRequiresPlatformSDK = true;
+
 		PublicIncludePathModuleNames.Add("Launch");
 
 		PrivateDependencyModuleNames.AddRange(new string[] { "Core", "CoreUObject", "zlib", "SwarmInterface", "Projects", "ApplicationCore" });
