Fixed printf.

diff --git a/Engine/Source/Programs/UnrealLightmass/Private/Launch/UnitTest.cpp b/Engine/Source/Programs/UnrealLightmass/Private/Launch/UnitTest.cpp
--- a/Engine/Source/Programs/UnrealLightmass/Private/Launch/UnitTest.cpp
+++ b/Engine/Source/Programs/UnrealLightmass/Private/Launch/UnitTest.cpp
@@ -239,7 +239,7 @@ void TestLightmass()
 
 
 	double Start = FPlatformTime::Seconds();
-	UE_LOG(LogLightmass, Display, TEXT("\nWaiting 4 seconds"), Start);
+	UE_LOG(LogLightmass, Display, TEXT("\nWaiting 4 seconds"));
 	FPlatformProcess::Sleep(4.0f);
 	UE_LOG(LogLightmass, Display, TEXT("%.2f seconds have passed, killing thread"), FPlatformTime::Seconds() - Start);
 	
