UnrealLightmass: Fixed an issue where Precomputed Visibility was broken due to inconsistent results between rtcOccluded1 and rtcIntersect1 when determining ray occlusion.

#Synced-CL 51880279

diff --git a/Engine/Source/Programs/UnrealLightmass/Private/Lighting/Embree.cpp b/Engine/Source/Programs/UnrealLightmass/Private/Lighting/Embree.cpp
--- a/Engine/Source/Programs/UnrealLightmass/Private/Lighting/Embree.cpp
+++ b/Engine/Source/Programs/UnrealLightmass/Private/Lighting/Embree.cpp
@@ -1121,7 +1121,12 @@ bool FEmbreeAggregateMesh::IntersectLightRay(
 	}
 
 #if USE_EMBREE_MAJOR_VERSION >= 3
-	if (EmbreeRayHit.ray.tfar >= 0.0f && EmbreeRayHit.hit.geomID != -1 && EmbreeRayHit.hit.primID != -1)
+	if (!bFindClosestIntersection)
+	{
+		// rtcOccluded1 : When occluded tfar is set to -infinity
+		ClosestIntersection.bIntersects = (EmbreeRayHit.ray.tfar < 0.0f);
+	}
+	else if (EmbreeRayHit.ray.tfar >= 0.0f && EmbreeRayHit.hit.geomID != -1 && EmbreeRayHit.hit.primID != -1)
 	{
 		FMinimalStaticLightingVertex EmbreeVertex;
 		EmbreeVertex.WorldPosition = LightRay.Start + LightRay.Direction * EmbreeRayHit.ray.tfar;
